Enabled UHT to digest GENERATED_BODY instead of GENERATED_UCLASS_BODY, GENERATED_USTRUCT_BODY, GENERATED_UINTERFACE_BODY or GENERATED_IINTERFACE_BODY, changed every occurence to the new syntax and fixed every warning that have fallen out of this change.

diff --git a/Engine/Plugins/2D/Paper2D/Source/Paper2D/Classes/PaperCharacter.h b/Engine/Plugins/2D/Paper2D/Source/Paper2D/Classes/PaperCharacter.h
--- a/Engine/Plugins/2D/Paper2D/Source/Paper2D/Classes/PaperCharacter.h
+++ b/Engine/Plugins/2D/Paper2D/Source/Paper2D/Classes/PaperCharacter.h
@@ -10,7 +10,9 @@
 UCLASS()
 class PAPER2D_API APaperCharacter : public ACharacter
 {
-	GENERATED_UCLASS_BODY()
+	GENERATED_BODY()
+public:
+	APaperCharacter(const FObjectInitializer& ObjectInitializer = FObjectInitializer::Get());
 	
 	// Name of the Sprite component
 	static FName SpriteComponentName;
diff --git a/Engine/Plugins/2D/Paper2D/Source/Paper2D/Classes/PaperFlipbook.h b/Engine/Plugins/2D/Paper2D/Source/Paper2D/Classes/PaperFlipbook.h
--- a/Engine/Plugins/2D/Paper2D/Source/Paper2D/Classes/PaperFlipbook.h
+++ b/Engine/Plugins/2D/Paper2D/Source/Paper2D/Classes/PaperFlipbook.h
@@ -8,8 +8,8 @@
 USTRUCT()
 struct FPaperFlipbookKeyFrame
 {
+	GENERATED_BODY()
 public:
-	GENERATED_USTRUCT_BODY()
 
 	UPROPERTY(Category=Sprite, EditAnywhere)
 	UPaperSprite* Sprite;
@@ -46,7 +46,9 @@ namespace EFlipbookCollisionMode
 UCLASS(BlueprintType, meta = (DisplayThumbnail = "true"))
 class PAPER2D_API UPaperFlipbook : public UObject
 {
-	GENERATED_UCLASS_BODY()
+	GENERATED_BODY()
+public:
+	UPaperFlipbook(const FObjectInitializer& ObjectInitializer = FObjectInitializer::Get());
 
 protected:
 	// The nominal frame rate to play this flipbook animation back at
diff --git a/Engine/Plugins/2D/Paper2D/Source/Paper2D/Classes/PaperFlipbookActor.h b/Engine/Plugins/2D/Paper2D/Source/Paper2D/Classes/PaperFlipbookActor.h
--- a/Engine/Plugins/2D/Paper2D/Source/Paper2D/Classes/PaperFlipbookActor.h
+++ b/Engine/Plugins/2D/Paper2D/Source/Paper2D/Classes/PaperFlipbookActor.h
@@ -14,7 +14,9 @@
 UCLASS(ComponentWrapperClass)
 class PAPER2D_API APaperFlipbookActor : public AActor
 {
-	GENERATED_UCLASS_BODY()
+	GENERATED_BODY()
+public:
+	APaperFlipbookActor(const FObjectInitializer& ObjectInitializer = FObjectInitializer::Get());
 
 private:
 	UPROPERTY(Category=Sprite, VisibleAnywhere, BlueprintReadOnly, meta=(ExposeFunctionCategories="Sprite,Rendering,Physics,Components|Flipbook", AllowPrivateAccess="true"))
diff --git a/Engine/Plugins/2D/Paper2D/Source/Paper2D/Classes/PaperFlipbookComponent.h b/Engine/Plugins/2D/Paper2D/Source/Paper2D/Classes/PaperFlipbookComponent.h
--- a/Engine/Plugins/2D/Paper2D/Source/Paper2D/Classes/PaperFlipbookComponent.h
+++ b/Engine/Plugins/2D/Paper2D/Source/Paper2D/Classes/PaperFlipbookComponent.h
@@ -11,7 +11,9 @@ DECLARE_DYNAMIC_MULTICAST_DELEGATE(FFlipbookFinishedPlaySignature);
 UCLASS(ShowCategories=(Mobility, ComponentReplication), ClassGroup=Paper2D, EarlyAccessPreview, meta=(BlueprintSpawnableComponent))
 class PAPER2D_API UPaperFlipbookComponent : public UMeshComponent
 {
-	GENERATED_UCLASS_BODY()
+	GENERATED_BODY()
+public:
+	UPaperFlipbookComponent(const FObjectInitializer& ObjectInitializer = FObjectInitializer::Get());
 
 protected:
 	/** Flipbook currently being played */
diff --git a/Engine/Plugins/2D/Paper2D/Source/Paper2D/Classes/PaperRuntimeSettings.h b/Engine/Plugins/2D/Paper2D/Source/Paper2D/Classes/PaperRuntimeSettings.h
--- a/Engine/Plugins/2D/Paper2D/Source/Paper2D/Classes/PaperRuntimeSettings.h
+++ b/Engine/Plugins/2D/Paper2D/Source/Paper2D/Classes/PaperRuntimeSettings.h
@@ -10,7 +10,9 @@
 UCLASS(config = Engine, defaultconfig)
 class PAPER2D_API UPaperRuntimeSettings : public UObject
 {
-	GENERATED_UCLASS_BODY()
+	GENERATED_BODY()
+public:
+	UPaperRuntimeSettings(const FObjectInitializer& ObjectInitializer = FObjectInitializer::Get());
 
 	// The default scaling factor between pixels and Unreal units (cm) to use for newly created sprite assets (e.g., 0.64 would make a 64 pixel wide sprite take up 100 cm)
 	UPROPERTY(GlobalConfig, EditAnywhere, Category=Settings)
diff --git a/Engine/Plugins/2D/Paper2D/Source/Paper2D/Classes/PaperSprite.h b/Engine/Plugins/2D/Paper2D/Source/Paper2D/Classes/PaperSprite.h
--- a/Engine/Plugins/2D/Paper2D/Source/Paper2D/Classes/PaperSprite.h
+++ b/Engine/Plugins/2D/Paper2D/Source/Paper2D/Classes/PaperSprite.h
@@ -22,7 +22,7 @@ namespace ESpriteCollisionMode
 USTRUCT()
 struct FPaperSpriteSocket
 {
-	GENERATED_USTRUCT_BODY()
+	GENERATED_BODY()
 
 	// Transform in pivot space (*not* texture space)
 	UPROPERTY(EditAnywhere, BlueprintReadOnly, Category=Sockets)
@@ -44,7 +44,9 @@ typedef TArray<class UTexture2D*, TInlineAllocator<4>> FAdditionalSpriteTextureA
 UCLASS(BlueprintType, meta=(DisplayThumbnail = "true"))
 class PAPER2D_API UPaperSprite : public UObject, public IInterface_CollisionDataProvider
 {
-	GENERATED_UCLASS_BODY()
+	GENERATED_BODY()
+public:
+	UPaperSprite(const FObjectInitializer& ObjectInitializer = FObjectInitializer::Get());
 
 protected:
 
diff --git a/Engine/Plugins/2D/Paper2D/Source/Paper2D/Classes/PaperSpriteActor.h b/Engine/Plugins/2D/Paper2D/Source/Paper2D/Classes/PaperSpriteActor.h
--- a/Engine/Plugins/2D/Paper2D/Source/Paper2D/Classes/PaperSpriteActor.h
+++ b/Engine/Plugins/2D/Paper2D/Source/Paper2D/Classes/PaperSpriteActor.h
@@ -14,7 +14,9 @@
 UCLASS(ComponentWrapperClass, meta = (ChildCanTick))
 class PAPER2D_API APaperSpriteActor : public AActor
 {
-	GENERATED_UCLASS_BODY()
+	GENERATED_BODY()
+public:
+	APaperSpriteActor(const FObjectInitializer& ObjectInitializer = FObjectInitializer::Get());
 
 private:
 	UPROPERTY(Category=Sprite, VisibleAnywhere, BlueprintReadOnly, meta=(ExposeFunctionCategories="Sprite,Rendering,Physics,Components|Sprite", AllowPrivateAccess="true"))
diff --git a/Engine/Plugins/2D/Paper2D/Source/Paper2D/Classes/PaperSpriteAtlas.h b/Engine/Plugins/2D/Paper2D/Source/Paper2D/Classes/PaperSpriteAtlas.h
--- a/Engine/Plugins/2D/Paper2D/Source/Paper2D/Classes/PaperSpriteAtlas.h
+++ b/Engine/Plugins/2D/Paper2D/Source/Paper2D/Classes/PaperSpriteAtlas.h
@@ -7,7 +7,7 @@
 USTRUCT()
 struct FPaperSpriteAtlasSlot
 {
-	GENERATED_USTRUCT_BODY()
+	GENERATED_BODY()
 
 	UPROPERTY()
 	TAssetPtr<class UPaperSprite> SpriteRef;
@@ -47,7 +47,9 @@ enum class EPaperSpriteAtlasPadding : uint8
 UCLASS(MinimalAPI, BlueprintType, Experimental, meta=(DisplayThumbnail = "true"))
 class UPaperSpriteAtlas : public UObject
 {
-	GENERATED_UCLASS_BODY()
+	GENERATED_BODY()
+public:
+	PAPER2D_API UPaperSpriteAtlas(const FObjectInitializer& ObjectInitializer = FObjectInitializer::Get());
 
 #if WITH_EDITORONLY_DATA
 	// Description of this atlas, which shows up in the content browser tooltip
diff --git a/Engine/Plugins/2D/Paper2D/Source/Paper2D/Classes/PaperSpriteComponent.h b/Engine/Plugins/2D/Paper2D/Source/Paper2D/Classes/PaperSpriteComponent.h
--- a/Engine/Plugins/2D/Paper2D/Source/Paper2D/Classes/PaperSpriteComponent.h
+++ b/Engine/Plugins/2D/Paper2D/Source/Paper2D/Classes/PaperSpriteComponent.h
@@ -18,7 +18,9 @@
 UCLASS(ShowCategories=(Mobility), ClassGroup=Paper2D, meta=(BlueprintSpawnableComponent))
 class PAPER2D_API UPaperSpriteComponent : public UMeshComponent
 {
-	GENERATED_UCLASS_BODY()
+	GENERATED_BODY()
+public:
+	UPaperSpriteComponent(const FObjectInitializer& ObjectInitializer = FObjectInitializer::Get());
 
 protected:
 	// The sprite asset used by this component
diff --git a/Engine/Plugins/2D/Paper2D/Source/Paper2D/Classes/PaperTerrainActor.h b/Engine/Plugins/2D/Paper2D/Source/Paper2D/Classes/PaperTerrainActor.h
--- a/Engine/Plugins/2D/Paper2D/Source/Paper2D/Classes/PaperTerrainActor.h
+++ b/Engine/Plugins/2D/Paper2D/Source/Paper2D/Classes/PaperTerrainActor.h
@@ -13,7 +13,9 @@
 UCLASS(BlueprintType, Experimental)
 class PAPER2D_API APaperTerrainActor : public AActor
 {
-	GENERATED_UCLASS_BODY()
+	GENERATED_BODY()
+public:
+	APaperTerrainActor(const FObjectInitializer& ObjectInitializer = FObjectInitializer::Get());
 
 private:
 	UPROPERTY()
diff --git a/Engine/Plugins/2D/Paper2D/Source/Paper2D/Classes/PaperTerrainComponent.h b/Engine/Plugins/2D/Paper2D/Source/Paper2D/Classes/PaperTerrainComponent.h
--- a/Engine/Plugins/2D/Paper2D/Source/Paper2D/Classes/PaperTerrainComponent.h
+++ b/Engine/Plugins/2D/Paper2D/Source/Paper2D/Classes/PaperTerrainComponent.h
@@ -44,7 +44,9 @@ struct FTerrainSegment
 UCLASS(BlueprintType, Experimental)
 class PAPER2D_API UPaperTerrainComponent : public UPrimitiveComponent
 {
-	GENERATED_UCLASS_BODY()
+	GENERATED_BODY()
+public:
+	UPaperTerrainComponent(const FObjectInitializer& ObjectInitializer = FObjectInitializer::Get());
 
 public:
 	/** The terrain material to apply to this component (set of rules for which sprites are used on different surfaces or the interior) */
diff --git a/Engine/Plugins/2D/Paper2D/Source/Paper2D/Classes/PaperTerrainMaterial.h b/Engine/Plugins/2D/Paper2D/Source/Paper2D/Classes/PaperTerrainMaterial.h
--- a/Engine/Plugins/2D/Paper2D/Source/Paper2D/Classes/PaperTerrainMaterial.h
+++ b/Engine/Plugins/2D/Paper2D/Source/Paper2D/Classes/PaperTerrainMaterial.h
@@ -8,7 +8,7 @@
 USTRUCT(BlueprintType)
 struct FPaperTerrainMaterialRule
 {
-	GENERATED_USTRUCT_BODY()
+	GENERATED_BODY()
 
 	// The sprite to use at the 'left' (closest to spline start) edge of the terrain segment
 	UPROPERTY(Category=Sprite, EditAnywhere, meta=(DisplayThumbnail="true"))
@@ -68,7 +68,9 @@ struct FPaperTerrainMaterialRule
 UCLASS(BlueprintType)
 class PAPER2D_API UPaperTerrainMaterial : public UDataAsset
 {
-	GENERATED_UCLASS_BODY()
+	GENERATED_BODY()
+public:
+	UPaperTerrainMaterial(const FObjectInitializer& ObjectInitializer = FObjectInitializer::Get());
 
 	UPROPERTY(Category=Sprite, EditAnywhere)
 	TArray<FPaperTerrainMaterialRule> Rules;
diff --git a/Engine/Plugins/2D/Paper2D/Source/Paper2D/Classes/PaperTerrainSplineComponent.h b/Engine/Plugins/2D/Paper2D/Source/Paper2D/Classes/PaperTerrainSplineComponent.h
--- a/Engine/Plugins/2D/Paper2D/Source/Paper2D/Classes/PaperTerrainSplineComponent.h
+++ b/Engine/Plugins/2D/Paper2D/Source/Paper2D/Classes/PaperTerrainSplineComponent.h
@@ -9,7 +9,9 @@
 UCLASS(BlueprintType, Experimental)
 class PAPER2D_API UPaperTerrainSplineComponent : public USplineComponent
 {
-	GENERATED_UCLASS_BODY()
+	GENERATED_BODY()
+public:
+	UPaperTerrainSplineComponent(const FObjectInitializer& ObjectInitializer = FObjectInitializer::Get());
 
 public:
 	// UObject interface
diff --git a/Engine/Plugins/2D/Paper2D/Source/Paper2D/Classes/PaperTileLayer.h b/Engine/Plugins/2D/Paper2D/Source/Paper2D/Classes/PaperTileLayer.h
--- a/Engine/Plugins/2D/Paper2D/Source/Paper2D/Classes/PaperTileLayer.h
+++ b/Engine/Plugins/2D/Paper2D/Source/Paper2D/Classes/PaperTileLayer.h
@@ -20,7 +20,7 @@ enum class EPaperTileFlags : uint32
 USTRUCT(BlueprintType)
 struct FPaperTileInfo
 {
-	GENERATED_USTRUCT_BODY()
+	GENERATED_BODY()
 
 	// The tile set that this tile comes from
 	UPROPERTY(EditAnywhere, BlueprintReadWrite, Category=Sprite)
@@ -101,7 +101,9 @@ struct FPaperTileInfo
 UCLASS()
 class PAPER2D_API UPaperTileLayer : public UObject
 {
-	GENERATED_UCLASS_BODY()
+	GENERATED_BODY()
+public:
+	UPaperTileLayer(const FObjectInitializer& ObjectInitializer = FObjectInitializer::Get());
 
 	// Name of the layer
 	UPROPERTY()
diff --git a/Engine/Plugins/2D/Paper2D/Source/Paper2D/Classes/PaperTileMap.h b/Engine/Plugins/2D/Paper2D/Source/Paper2D/Classes/PaperTileMap.h
--- a/Engine/Plugins/2D/Paper2D/Source/Paper2D/Classes/PaperTileMap.h
+++ b/Engine/Plugins/2D/Paper2D/Source/Paper2D/Classes/PaperTileMap.h
@@ -31,7 +31,9 @@ namespace ETileMapProjectionMode
 UCLASS()
 class PAPER2D_API UPaperTileMap : public UObject
 {
-	GENERATED_UCLASS_BODY()
+	GENERATED_BODY()
+public:
+	UPaperTileMap(const FObjectInitializer& ObjectInitializer = FObjectInitializer::Get());
 
 	// Width of map (in tiles)
 	UPROPERTY(Category=Setup, EditAnywhere, BlueprintReadOnly, meta=(UIMin=1, ClampMin=1))
diff --git a/Engine/Plugins/2D/Paper2D/Source/Paper2D/Classes/PaperTileMapActor.h b/Engine/Plugins/2D/Paper2D/Source/Paper2D/Classes/PaperTileMapActor.h
--- a/Engine/Plugins/2D/Paper2D/Source/Paper2D/Classes/PaperTileMapActor.h
+++ b/Engine/Plugins/2D/Paper2D/Source/Paper2D/Classes/PaperTileMapActor.h
@@ -15,7 +15,9 @@
 UCLASS(ComponentWrapperClass)
 class PAPER2D_API APaperTileMapActor : public AActor
 {
-	GENERATED_UCLASS_BODY()
+	GENERATED_BODY()
+public:
+	APaperTileMapActor(const FObjectInitializer& ObjectInitializer = FObjectInitializer::Get());
 
 private:
 	UPROPERTY(Category=TileMapActor, VisibleAnywhere, BlueprintReadOnly, meta=(ExposeFunctionCategories="Sprite,Rendering,Physics,Components|Sprite", AllowPrivateAccess="true"))
diff --git a/Engine/Plugins/2D/Paper2D/Source/Paper2D/Classes/PaperTileMapComponent.h b/Engine/Plugins/2D/Paper2D/Source/Paper2D/Classes/PaperTileMapComponent.h
--- a/Engine/Plugins/2D/Paper2D/Source/Paper2D/Classes/PaperTileMapComponent.h
+++ b/Engine/Plugins/2D/Paper2D/Source/Paper2D/Classes/PaperTileMapComponent.h
@@ -20,7 +20,9 @@
 UCLASS(hideCategories=Object, ClassGroup=Paper2D, Experimental, meta=(BlueprintSpawnableComponent))
 class PAPER2D_API UPaperTileMapComponent : public UMeshComponent
 {
-	GENERATED_UCLASS_BODY()
+	GENERATED_BODY()
+public:
+	UPaperTileMapComponent(const FObjectInitializer& ObjectInitializer = FObjectInitializer::Get());
 
 private:
 	UPROPERTY()
diff --git a/Engine/Plugins/2D/Paper2D/Source/Paper2D/Classes/PaperTileSet.h b/Engine/Plugins/2D/Paper2D/Source/Paper2D/Classes/PaperTileSet.h
--- a/Engine/Plugins/2D/Paper2D/Source/Paper2D/Classes/PaperTileSet.h
+++ b/Engine/Plugins/2D/Paper2D/Source/Paper2D/Classes/PaperTileSet.h
@@ -10,7 +10,9 @@
 UCLASS()
 class PAPER2D_API UPaperTileSet : public UObject
 {
-	GENERATED_UCLASS_BODY()
+	GENERATED_BODY()
+public:
+	UPaperTileSet(const FObjectInitializer& ObjectInitializer = FObjectInitializer::Get());
 
 	// The width of a single tile (in pixels)
 	UPROPERTY(Category=TileSet, EditAnywhere, meta=(UIMin=1, ClampMin=1))
diff --git a/Engine/Plugins/2D/Paper2D/Source/Paper2D/Classes/SpriteDrawCall.h b/Engine/Plugins/2D/Paper2D/Source/Paper2D/Classes/SpriteDrawCall.h
--- a/Engine/Plugins/2D/Paper2D/Source/Paper2D/Classes/SpriteDrawCall.h
+++ b/Engine/Plugins/2D/Paper2D/Source/Paper2D/Classes/SpriteDrawCall.h
@@ -8,8 +8,8 @@
 USTRUCT()
 struct FSpriteDrawCallRecord
 {
+	GENERATED_BODY()
 public:
-	GENERATED_USTRUCT_BODY()
 
 	UPROPERTY(Category=Sprite, EditAnywhere)
 	FVector Destination;
diff --git a/Engine/Plugins/2D/Paper2D/Source/Paper2D/Classes/SpriteEditorOnlyTypes.h b/Engine/Plugins/2D/Paper2D/Source/Paper2D/Classes/SpriteEditorOnlyTypes.h
--- a/Engine/Plugins/2D/Paper2D/Source/Paper2D/Classes/SpriteEditorOnlyTypes.h
+++ b/Engine/Plugins/2D/Paper2D/Source/Paper2D/Classes/SpriteEditorOnlyTypes.h
@@ -9,7 +9,7 @@
 USTRUCT()
 struct FSpritePolygon
 {
-	GENERATED_USTRUCT_BODY()
+	GENERATED_BODY()
 
 	FSpritePolygon() 
 	: bNegativeWinding(false)
@@ -82,7 +82,7 @@ namespace ESpriteSubdivisionMode
 USTRUCT()
 struct FSpritePolygonCollection
 {
-	GENERATED_USTRUCT_BODY()
+	GENERATED_BODY()
 
 	UPROPERTY(Category=PolygonData, EditAnywhere, AdvancedDisplay)
 	TArray<FSpritePolygon> Polygons;
@@ -147,7 +147,7 @@ struct FSpritePolygonCollection
 USTRUCT()
 struct FSpriteAssetInitParameters
 {
-	GENERATED_USTRUCT_BODY()
+	GENERATED_BODY()
 
 	FSpriteAssetInitParameters()
 		: Texture(nullptr)
diff --git a/Engine/Plugins/2D/Paper2D/Source/Paper2D/Private/Rendering/PaperBatchComponent.h b/Engine/Plugins/2D/Paper2D/Source/Paper2D/Private/Rendering/PaperBatchComponent.h
--- a/Engine/Plugins/2D/Paper2D/Source/Paper2D/Private/Rendering/PaperBatchComponent.h
+++ b/Engine/Plugins/2D/Paper2D/Source/Paper2D/Private/Rendering/PaperBatchComponent.h
@@ -8,7 +8,9 @@
 UCLASS()
 class UPaperBatchComponent : public UPrimitiveComponent
 {
-	GENERATED_UCLASS_BODY()
+	GENERATED_BODY()
+public:
+	UPaperBatchComponent(const FObjectInitializer& ObjectInitializer = FObjectInitializer::Get());
 
 	// UActorComponent interface
 	virtual void TickComponent(float DeltaTime, enum ELevelTick TickType, FActorComponentTickFunction* ThisTickFunction) override;
diff --git a/Engine/Plugins/2D/Paper2D/Source/Paper2D/Public/MaterialExpressionSpriteTextureSampler.h b/Engine/Plugins/2D/Paper2D/Source/Paper2D/Public/MaterialExpressionSpriteTextureSampler.h
--- a/Engine/Plugins/2D/Paper2D/Source/Paper2D/Public/MaterialExpressionSpriteTextureSampler.h
+++ b/Engine/Plugins/2D/Paper2D/Source/Paper2D/Public/MaterialExpressionSpriteTextureSampler.h
@@ -9,7 +9,9 @@
 UCLASS(hideCategories=(MaterialExpressionTextureSampleParameter, MaterialExpressionSpriteTextureSampler))
 class PAPER2D_API UMaterialExpressionSpriteTextureSampler : public UMaterialExpressionTextureSampleParameter2D
 {
-	GENERATED_UCLASS_BODY()
+	GENERATED_BODY()
+public:
+	UMaterialExpressionSpriteTextureSampler(const FObjectInitializer& ObjectInitializer = FObjectInitializer::Get());
 
 	// Is this a sampler for the default SourceTexture or the AdditionalSourceTextures list?
 	UPROPERTY(EditAnywhere, Category=Paper2D)
diff --git a/Engine/Plugins/2D/Paper2D/Source/Paper2DEditor/Classes/PaperFlipbookActorFactory.h b/Engine/Plugins/2D/Paper2D/Source/Paper2DEditor/Classes/PaperFlipbookActorFactory.h
--- a/Engine/Plugins/2D/Paper2D/Source/Paper2DEditor/Classes/PaperFlipbookActorFactory.h
+++ b/Engine/Plugins/2D/Paper2D/Source/Paper2DEditor/Classes/PaperFlipbookActorFactory.h
@@ -7,7 +7,9 @@
 UCLASS()
 class UPaperFlipbookActorFactory : public UActorFactory
 {
-	GENERATED_UCLASS_BODY()
+	GENERATED_BODY()
+public:
+	UPaperFlipbookActorFactory(const FObjectInitializer& ObjectInitializer = FObjectInitializer::Get());
 
 	// Begin UActorFactory Interface
 	virtual void PostSpawnActor(UObject* Asset, AActor* NewActor) override;
diff --git a/Engine/Plugins/2D/Paper2D/Source/Paper2DEditor/Classes/PaperFlipbookFactory.h b/Engine/Plugins/2D/Paper2D/Source/Paper2DEditor/Classes/PaperFlipbookFactory.h
--- a/Engine/Plugins/2D/Paper2D/Source/Paper2DEditor/Classes/PaperFlipbookFactory.h
+++ b/Engine/Plugins/2D/Paper2D/Source/Paper2DEditor/Classes/PaperFlipbookFactory.h
@@ -10,7 +10,9 @@
 UCLASS()
 class PAPER2DEDITOR_API UPaperFlipbookFactory : public UFactory
 {
-	GENERATED_UCLASS_BODY()
+	GENERATED_BODY()
+public:
+	UPaperFlipbookFactory(const FObjectInitializer& ObjectInitializer = FObjectInitializer::Get());
 
 	TArray<FPaperFlipbookKeyFrame> KeyFrames;
 
diff --git a/Engine/Plugins/2D/Paper2D/Source/Paper2DEditor/Classes/PaperFlipbookThumbnailRenderer.h b/Engine/Plugins/2D/Paper2D/Source/Paper2DEditor/Classes/PaperFlipbookThumbnailRenderer.h
--- a/Engine/Plugins/2D/Paper2D/Source/Paper2DEditor/Classes/PaperFlipbookThumbnailRenderer.h
+++ b/Engine/Plugins/2D/Paper2D/Source/Paper2DEditor/Classes/PaperFlipbookThumbnailRenderer.h
@@ -7,7 +7,9 @@
 UCLASS()
 class UPaperFlipbookThumbnailRenderer : public UPaperSpriteThumbnailRenderer
 {
-	GENERATED_UCLASS_BODY()
+	GENERATED_BODY()
+public:
+	UPaperFlipbookThumbnailRenderer(const FObjectInitializer& ObjectInitializer = FObjectInitializer::Get());
 
 	// UThumbnailRenderer interface
 	virtual void Draw(UObject* Object, int32 X, int32 Y, uint32 Width, uint32 Height, FRenderTarget*, FCanvas* Canvas) override;
diff --git a/Engine/Plugins/2D/Paper2D/Source/Paper2DEditor/Classes/PaperSpriteActorFactory.h b/Engine/Plugins/2D/Paper2D/Source/Paper2DEditor/Classes/PaperSpriteActorFactory.h
--- a/Engine/Plugins/2D/Paper2D/Source/Paper2DEditor/Classes/PaperSpriteActorFactory.h
+++ b/Engine/Plugins/2D/Paper2D/Source/Paper2DEditor/Classes/PaperSpriteActorFactory.h
@@ -7,7 +7,9 @@
 UCLASS()
 class UPaperSpriteActorFactory : public UActorFactory
 {
-	GENERATED_UCLASS_BODY()
+	GENERATED_BODY()
+public:
+	UPaperSpriteActorFactory(const FObjectInitializer& ObjectInitializer = FObjectInitializer::Get());
 
 	// Begin UActorFactory Interface
 	virtual void PostSpawnActor(UObject* Asset, AActor* NewActor) override;
diff --git a/Engine/Plugins/2D/Paper2D/Source/Paper2DEditor/Classes/PaperSpriteFactory.h b/Engine/Plugins/2D/Paper2D/Source/Paper2DEditor/Classes/PaperSpriteFactory.h
--- a/Engine/Plugins/2D/Paper2D/Source/Paper2DEditor/Classes/PaperSpriteFactory.h
+++ b/Engine/Plugins/2D/Paper2D/Source/Paper2DEditor/Classes/PaperSpriteFactory.h
@@ -11,7 +11,9 @@
 UCLASS()
 class UPaperSpriteFactory : public UFactory
 {
-	GENERATED_UCLASS_BODY()
+	GENERATED_BODY()
+public:
+	UPaperSpriteFactory(const FObjectInitializer& ObjectInitializer = FObjectInitializer::Get());
 
 	// Initial texture to create the sprite from (Can be nullptr)
 	class UTexture2D* InitialTexture;
diff --git a/Engine/Plugins/2D/Paper2D/Source/Paper2DEditor/Classes/PaperSpriteThumbnailRenderer.h b/Engine/Plugins/2D/Paper2D/Source/Paper2DEditor/Classes/PaperSpriteThumbnailRenderer.h
--- a/Engine/Plugins/2D/Paper2D/Source/Paper2DEditor/Classes/PaperSpriteThumbnailRenderer.h
+++ b/Engine/Plugins/2D/Paper2D/Source/Paper2DEditor/Classes/PaperSpriteThumbnailRenderer.h
@@ -7,7 +7,9 @@
 UCLASS()
 class UPaperSpriteThumbnailRenderer : public UDefaultSizedThumbnailRenderer
 {
-	GENERATED_UCLASS_BODY()
+	GENERATED_BODY()
+public:
+	UPaperSpriteThumbnailRenderer(const FObjectInitializer& ObjectInitializer = FObjectInitializer::Get());
 
 	// UThumbnailRenderer interface
 	virtual void Draw(UObject* Object, int32 X, int32 Y, uint32 Width, uint32 Height, FRenderTarget*, FCanvas* Canvas) override;
diff --git a/Engine/Plugins/2D/Paper2D/Source/Paper2DEditor/Classes/PaperTileMapFactory.h b/Engine/Plugins/2D/Paper2D/Source/Paper2DEditor/Classes/PaperTileMapFactory.h
--- a/Engine/Plugins/2D/Paper2D/Source/Paper2DEditor/Classes/PaperTileMapFactory.h
+++ b/Engine/Plugins/2D/Paper2D/Source/Paper2DEditor/Classes/PaperTileMapFactory.h
@@ -11,7 +11,9 @@
 UCLASS()
 class UPaperTileMapFactory : public UFactory
 {
-	GENERATED_UCLASS_BODY()
+	GENERATED_BODY()
+public:
+	UPaperTileMapFactory(const FObjectInitializer& ObjectInitializer = FObjectInitializer::Get());
 
 	// Initial tile set to create the tile map from (Can be nullptr)
 	class UPaperTileSet* InitialTileSet;
diff --git a/Engine/Plugins/2D/Paper2D/Source/Paper2DEditor/Classes/PaperTileMapPromotionFactory.h b/Engine/Plugins/2D/Paper2D/Source/Paper2DEditor/Classes/PaperTileMapPromotionFactory.h
--- a/Engine/Plugins/2D/Paper2D/Source/Paper2DEditor/Classes/PaperTileMapPromotionFactory.h
+++ b/Engine/Plugins/2D/Paper2D/Source/Paper2DEditor/Classes/PaperTileMapPromotionFactory.h
@@ -11,7 +11,9 @@
 UCLASS()
 class UPaperTileMapPromotionFactory : public UFactory
 {
-	GENERATED_UCLASS_BODY()
+	GENERATED_BODY()
+public:
+	UPaperTileMapPromotionFactory(const FObjectInitializer& ObjectInitializer = FObjectInitializer::Get());
 
 	// Object being promoted to an asset
 	UPROPERTY()
diff --git a/Engine/Plugins/2D/Paper2D/Source/Paper2DEditor/Classes/PaperTileSetFactory.h b/Engine/Plugins/2D/Paper2D/Source/Paper2DEditor/Classes/PaperTileSetFactory.h
--- a/Engine/Plugins/2D/Paper2D/Source/Paper2DEditor/Classes/PaperTileSetFactory.h
+++ b/Engine/Plugins/2D/Paper2D/Source/Paper2DEditor/Classes/PaperTileSetFactory.h
@@ -11,7 +11,9 @@
 UCLASS()
 class UPaperTileSetFactory : public UFactory
 {
-	GENERATED_UCLASS_BODY()
+	GENERATED_BODY()
+public:
+	UPaperTileSetFactory(const FObjectInitializer& ObjectInitializer = FObjectInitializer::Get());
 
 	// Initial texture to create the tile set from (Can be nullptr)
 	class UTexture2D* InitialTexture;
diff --git a/Engine/Plugins/2D/Paper2D/Source/Paper2DEditor/Classes/PaperTileSetThumbnailRenderer.h b/Engine/Plugins/2D/Paper2D/Source/Paper2DEditor/Classes/PaperTileSetThumbnailRenderer.h
--- a/Engine/Plugins/2D/Paper2D/Source/Paper2DEditor/Classes/PaperTileSetThumbnailRenderer.h
+++ b/Engine/Plugins/2D/Paper2D/Source/Paper2DEditor/Classes/PaperTileSetThumbnailRenderer.h
@@ -7,7 +7,9 @@
 UCLASS()
 class UPaperTileSetThumbnailRenderer : public UDefaultSizedThumbnailRenderer
 {
-	GENERATED_UCLASS_BODY()
+	GENERATED_BODY()
+public:
+	UPaperTileSetThumbnailRenderer(const FObjectInitializer& ObjectInitializer = FObjectInitializer::Get());
 
 	// UThumbnailRenderer interface
 	virtual void Draw(UObject* Object, int32 X, int32 Y, uint32 Width, uint32 Height, FRenderTarget*, FCanvas* Canvas) override;
diff --git a/Engine/Plugins/2D/Paper2D/Source/Paper2DEditor/Classes/SpriteEditorTools/PSE_DeleteTool.h b/Engine/Plugins/2D/Paper2D/Source/Paper2DEditor/Classes/SpriteEditorTools/PSE_DeleteTool.h
--- a/Engine/Plugins/2D/Paper2D/Source/Paper2DEditor/Classes/SpriteEditorTools/PSE_DeleteTool.h
+++ b/Engine/Plugins/2D/Paper2D/Source/Paper2DEditor/Classes/SpriteEditorTools/PSE_DeleteTool.h
@@ -7,5 +7,7 @@
 UCLASS()
 class UPSE_DeleteTool : public USpriteEditorTool
 {
-	GENERATED_UCLASS_BODY()
+	GENERATED_BODY()
+public:
+	UPSE_DeleteTool(const FObjectInitializer& ObjectInitializer = FObjectInitializer::Get());
 };
diff --git a/Engine/Plugins/2D/Paper2D/Source/Paper2DEditor/Classes/SpriteEditorTools/PSE_LineTool.h b/Engine/Plugins/2D/Paper2D/Source/Paper2DEditor/Classes/SpriteEditorTools/PSE_LineTool.h
--- a/Engine/Plugins/2D/Paper2D/Source/Paper2DEditor/Classes/SpriteEditorTools/PSE_LineTool.h
+++ b/Engine/Plugins/2D/Paper2D/Source/Paper2DEditor/Classes/SpriteEditorTools/PSE_LineTool.h
@@ -7,5 +7,7 @@
 UCLASS()
 class UPSE_LineTool : public USpriteEditorTool
 {
-	GENERATED_UCLASS_BODY()
+	GENERATED_BODY()
+public:
+	UPSE_LineTool(const FObjectInitializer& ObjectInitializer = FObjectInitializer::Get());
 };
diff --git a/Engine/Plugins/2D/Paper2D/Source/Paper2DEditor/Classes/SpriteEditorTools/PSE_SplitEdgeTool.h b/Engine/Plugins/2D/Paper2D/Source/Paper2DEditor/Classes/SpriteEditorTools/PSE_SplitEdgeTool.h
--- a/Engine/Plugins/2D/Paper2D/Source/Paper2DEditor/Classes/SpriteEditorTools/PSE_SplitEdgeTool.h
+++ b/Engine/Plugins/2D/Paper2D/Source/Paper2DEditor/Classes/SpriteEditorTools/PSE_SplitEdgeTool.h
@@ -7,5 +7,7 @@
 UCLASS()
 class UPSE_SplitEdgeTool : public USpriteEditorTool
 {
-	GENERATED_UCLASS_BODY()
+	GENERATED_BODY()
+public:
+	UPSE_SplitEdgeTool(const FObjectInitializer& ObjectInitializer = FObjectInitializer::Get());
 };
diff --git a/Engine/Plugins/2D/Paper2D/Source/Paper2DEditor/Classes/SpriteEditorTools/SpriteEditorTool.h b/Engine/Plugins/2D/Paper2D/Source/Paper2DEditor/Classes/SpriteEditorTools/SpriteEditorTool.h
--- a/Engine/Plugins/2D/Paper2D/Source/Paper2DEditor/Classes/SpriteEditorTools/SpriteEditorTool.h
+++ b/Engine/Plugins/2D/Paper2D/Source/Paper2DEditor/Classes/SpriteEditorTools/SpriteEditorTool.h
@@ -7,5 +7,7 @@
 UCLASS()
 class USpriteEditorTool : public UObject
 {
-	GENERATED_UCLASS_BODY()
+	GENERATED_BODY()
+public:
+	USpriteEditorTool(const FObjectInitializer& ObjectInitializer = FObjectInitializer::Get());
 };
diff --git a/Engine/Plugins/2D/Paper2D/Source/Paper2DEditor/Classes/TerrainSplineActorFactory.h b/Engine/Plugins/2D/Paper2D/Source/Paper2DEditor/Classes/TerrainSplineActorFactory.h
--- a/Engine/Plugins/2D/Paper2D/Source/Paper2DEditor/Classes/TerrainSplineActorFactory.h
+++ b/Engine/Plugins/2D/Paper2D/Source/Paper2DEditor/Classes/TerrainSplineActorFactory.h
@@ -7,7 +7,9 @@
 UCLASS()
 class UTerrainSplineActorFactory : public UActorFactory
 {
-	GENERATED_UCLASS_BODY()
+	GENERATED_BODY()
+public:
+	UTerrainSplineActorFactory(const FObjectInitializer& ObjectInitializer = FObjectInitializer::Get());
 
 	// UActorFactory interface
 	virtual void PostSpawnActor(UObject* Asset, AActor* NewActor) override;
diff --git a/Engine/Plugins/2D/Paper2D/Source/Paper2DEditor/Classes/TileMapActorFactory.h b/Engine/Plugins/2D/Paper2D/Source/Paper2DEditor/Classes/TileMapActorFactory.h
--- a/Engine/Plugins/2D/Paper2D/Source/Paper2DEditor/Classes/TileMapActorFactory.h
+++ b/Engine/Plugins/2D/Paper2D/Source/Paper2DEditor/Classes/TileMapActorFactory.h
@@ -7,7 +7,9 @@
 UCLASS()
 class UTileMapActorFactory : public UActorFactory
 {
-	GENERATED_UCLASS_BODY()
+	GENERATED_BODY()
+public:
+	UTileMapActorFactory(const FObjectInitializer& ObjectInitializer = FObjectInitializer::Get());
 
 	// UActorFactory interface
 	virtual void PostSpawnActor(UObject* Asset, AActor* NewActor) override;
diff --git a/Engine/Plugins/2D/Paper2D/Source/Paper2DEditor/Classes/TileMapAssetImportData.h b/Engine/Plugins/2D/Paper2D/Source/Paper2DEditor/Classes/TileMapAssetImportData.h
--- a/Engine/Plugins/2D/Paper2D/Source/Paper2DEditor/Classes/TileMapAssetImportData.h
+++ b/Engine/Plugins/2D/Paper2D/Source/Paper2DEditor/Classes/TileMapAssetImportData.h
@@ -7,7 +7,7 @@
 USTRUCT()
 struct FTileSetImportMapping
 {
-	GENERATED_USTRUCT_BODY()
+	GENERATED_BODY()
 
 	UPROPERTY()
 	FString SourceName;
@@ -25,7 +25,9 @@ struct FTileSetImportMapping
 UCLASS()
 class PAPER2DEDITOR_API UTileMapAssetImportData : public UAssetImportData
 {
-	GENERATED_UCLASS_BODY()
+	GENERATED_BODY()
+public:
+	UTileMapAssetImportData(const FObjectInitializer& ObjectInitializer = FObjectInitializer::Get());
 
 	UPROPERTY()
 	TArray<FTileSetImportMapping> TileSetMap;
diff --git a/Engine/Plugins/2D/Paper2D/Source/Paper2DEditor/Private/Atlasing/PaperSpriteAtlasFactory.h b/Engine/Plugins/2D/Paper2D/Source/Paper2DEditor/Private/Atlasing/PaperSpriteAtlasFactory.h
--- a/Engine/Plugins/2D/Paper2D/Source/Paper2DEditor/Private/Atlasing/PaperSpriteAtlasFactory.h
+++ b/Engine/Plugins/2D/Paper2D/Source/Paper2DEditor/Private/Atlasing/PaperSpriteAtlasFactory.h
@@ -7,7 +7,9 @@
 UCLASS()
 class UPaperSpriteAtlasFactory : public UFactory
 {
-	GENERATED_UCLASS_BODY()
+	GENERATED_BODY()
+public:
+	UPaperSpriteAtlasFactory(const FObjectInitializer& ObjectInitializer = FObjectInitializer::Get());
 
 	// UFactory interface
 	virtual UObject* FactoryCreateNew(UClass* Class, UObject* InParent, FName Name, EObjectFlags Flags, UObject* Context, FFeedbackContext* Warn) override;
diff --git a/Engine/Plugins/2D/Paper2D/Source/PaperJsonImporter/Classes/PaperJsonImporterFactory.h b/Engine/Plugins/2D/Paper2D/Source/PaperJsonImporter/Classes/PaperJsonImporterFactory.h
--- a/Engine/Plugins/2D/Paper2D/Source/PaperJsonImporter/Classes/PaperJsonImporterFactory.h
+++ b/Engine/Plugins/2D/Paper2D/Source/PaperJsonImporter/Classes/PaperJsonImporterFactory.h
@@ -8,7 +8,9 @@
 UCLASS()
 class UPaperJsonImporterFactory : public UFactory
 {
-	GENERATED_UCLASS_BODY()
+	GENERATED_BODY()
+public:
+	UPaperJsonImporterFactory(const FObjectInitializer& ObjectInitializer = FObjectInitializer::Get());
 
 	// UFactory interface
 	virtual FText GetToolTip() const override;
diff --git a/Engine/Plugins/2D/Paper2D/Source/PaperJsonImporter/Private/PaperSpriteSheet.h b/Engine/Plugins/2D/Paper2D/Source/PaperJsonImporter/Private/PaperSpriteSheet.h
--- a/Engine/Plugins/2D/Paper2D/Source/PaperJsonImporter/Private/PaperSpriteSheet.h
+++ b/Engine/Plugins/2D/Paper2D/Source/PaperJsonImporter/Private/PaperSpriteSheet.h
@@ -7,7 +7,9 @@
 UCLASS(BlueprintType, meta = (DisplayThumbnail = "true"))
 class UPaperSpriteSheet : public UObject
 {
-	GENERATED_UCLASS_BODY()
+	GENERATED_BODY()
+public:
+	UPaperSpriteSheet(const FObjectInitializer& ObjectInitializer = FObjectInitializer::Get());
 
 public:
 	// The names of sprites during import
diff --git a/Engine/Plugins/2D/Paper2D/Source/PaperJsonImporter/Private/ReimportPaperJsonImporterFactory.h b/Engine/Plugins/2D/Paper2D/Source/PaperJsonImporter/Private/ReimportPaperJsonImporterFactory.h
--- a/Engine/Plugins/2D/Paper2D/Source/PaperJsonImporter/Private/ReimportPaperJsonImporterFactory.h
+++ b/Engine/Plugins/2D/Paper2D/Source/PaperJsonImporter/Private/ReimportPaperJsonImporterFactory.h
@@ -7,7 +7,9 @@
 UCLASS()
 class UReimportPaperJsonImporterFactory : public UPaperJsonImporterFactory, public FReimportHandler
 {
-	GENERATED_UCLASS_BODY()
+	GENERATED_BODY()
+public:
+	UReimportPaperJsonImporterFactory(const FObjectInitializer& ObjectInitializer = FObjectInitializer::Get());
 
 	// Begin FReimportHandler interface
 	virtual bool CanReimport(UObject* Obj, TArray<FString>& OutFilenames) override;
diff --git a/Engine/Plugins/2D/Paper2D/Source/PaperTiledImporter/Classes/PaperTiledImporterFactory.h b/Engine/Plugins/2D/Paper2D/Source/PaperTiledImporter/Classes/PaperTiledImporterFactory.h
--- a/Engine/Plugins/2D/Paper2D/Source/PaperTiledImporter/Classes/PaperTiledImporterFactory.h
+++ b/Engine/Plugins/2D/Paper2D/Source/PaperTiledImporter/Classes/PaperTiledImporterFactory.h
@@ -22,7 +22,9 @@ namespace ETiledOrientation
 UCLASS()
 class UPaperTiledImporterFactory : public UFactory, public FReimportHandler
 {
-	GENERATED_UCLASS_BODY()
+	GENERATED_BODY()
+public:
+	UPaperTiledImporterFactory(const FObjectInitializer& ObjectInitializer = FObjectInitializer::Get());
 
 	// UFactory interface
 	virtual FText GetToolTip() const override;
diff --git a/Engine/Plugins/Developer/UObjectPlugin/Source/UObjectPlugin/Classes/MyPluginObject.h b/Engine/Plugins/Developer/UObjectPlugin/Source/UObjectPlugin/Classes/MyPluginObject.h
--- a/Engine/Plugins/Developer/UObjectPlugin/Source/UObjectPlugin/Classes/MyPluginObject.h
+++ b/Engine/Plugins/Developer/UObjectPlugin/Source/UObjectPlugin/Classes/MyPluginObject.h
@@ -11,7 +11,7 @@
 USTRUCT()
 struct FMyPluginStruct
 {
-	GENERATED_USTRUCT_BODY()
+	GENERATED_BODY()
  
 	UPROPERTY()
 	FString TestString;
@@ -24,7 +24,9 @@ struct FMyPluginStruct
 UCLASS()
 class UMyPluginObject : public UObject
 {
-	GENERATED_UCLASS_BODY()
+	GENERATED_BODY()
+public:
+	UMyPluginObject(const FObjectInitializer& ObjectInitializer = FObjectInitializer::Get());
 
 public:
 
diff --git a/Engine/Plugins/Editor/MacGraphicsSwitching/Source/MacGraphicsSwitching/Private/MacGraphicsSwitchingSettings.h b/Engine/Plugins/Editor/MacGraphicsSwitching/Source/MacGraphicsSwitching/Private/MacGraphicsSwitchingSettings.h
--- a/Engine/Plugins/Editor/MacGraphicsSwitching/Source/MacGraphicsSwitching/Private/MacGraphicsSwitchingSettings.h
+++ b/Engine/Plugins/Editor/MacGraphicsSwitching/Source/MacGraphicsSwitching/Private/MacGraphicsSwitchingSettings.h
@@ -6,9 +6,10 @@
 UCLASS(Config=EditorGameAgnostic)
 class UMacGraphicsSwitchingSettings : public UObject
 {
-	GENERATED_UCLASS_BODY()
+	GENERATED_BODY()
 	
 public:
+	UMacGraphicsSwitchingSettings(const FObjectInitializer& ObjectInitializer = FObjectInitializer::Get());
 	
 	UPROPERTY(Config, EditAnywhere, Category=OpenGL)
 	int32 DisplayMask;
diff --git a/Engine/Plugins/Editor/SpeedTreeImporter/Source/SpeedTreeImporter/Classes/ReimportSpeedTreeFactory.h b/Engine/Plugins/Editor/SpeedTreeImporter/Source/SpeedTreeImporter/Classes/ReimportSpeedTreeFactory.h
--- a/Engine/Plugins/Editor/SpeedTreeImporter/Source/SpeedTreeImporter/Classes/ReimportSpeedTreeFactory.h
+++ b/Engine/Plugins/Editor/SpeedTreeImporter/Source/SpeedTreeImporter/Classes/ReimportSpeedTreeFactory.h
@@ -11,7 +11,9 @@
 UCLASS(collapsecategories)
 class UReimportSpeedTreeFactory : public USpeedTreeImportFactory, public FReimportHandler
 {
-	GENERATED_UCLASS_BODY()
+	GENERATED_BODY()
+public:
+	UReimportSpeedTreeFactory(const FObjectInitializer& ObjectInitializer = FObjectInitializer::Get());
 
 	// Begin FReimportHandler interface
 	virtual bool CanReimport(UObject* Obj, TArray<FString>& OutFilenames) override;
diff --git a/Engine/Plugins/Editor/SpeedTreeImporter/Source/SpeedTreeImporter/Classes/SpeedTreeImportFactory.h b/Engine/Plugins/Editor/SpeedTreeImporter/Source/SpeedTreeImporter/Classes/SpeedTreeImportFactory.h
--- a/Engine/Plugins/Editor/SpeedTreeImporter/Source/SpeedTreeImporter/Classes/SpeedTreeImportFactory.h
+++ b/Engine/Plugins/Editor/SpeedTreeImporter/Source/SpeedTreeImporter/Classes/SpeedTreeImportFactory.h
@@ -10,7 +10,9 @@
 UCLASS(hidecategories=Object)
 class USpeedTreeImportFactory : public UFactory
 {
-	GENERATED_UCLASS_BODY()
+	GENERATED_BODY()
+public:
+	USpeedTreeImportFactory(const FObjectInitializer& ObjectInitializer = FObjectInitializer::Get());
 
 	// Begin UFactory Interface
 	virtual FText GetDisplayName() const override;
diff --git a/Engine/Plugins/Experimental/HTML5Networking/Source/HTML5Networking/Classes/WebSocketConnection.h b/Engine/Plugins/Experimental/HTML5Networking/Source/HTML5Networking/Classes/WebSocketConnection.h
--- a/Engine/Plugins/Experimental/HTML5Networking/Source/HTML5Networking/Classes/WebSocketConnection.h
+++ b/Engine/Plugins/Experimental/HTML5Networking/Source/HTML5Networking/Classes/WebSocketConnection.h
@@ -8,7 +8,9 @@
 UCLASS(transient, config = Engine)
 class HTML5NETWORKING_API UWebSocketConnection : public UNetConnection
 {
-	GENERATED_UCLASS_BODY()
+	GENERATED_BODY()
+public:
+	UWebSocketConnection(const FObjectInitializer& ObjectInitializer = FObjectInitializer::Get());
 
 	class FWebSocket*			WebSocket; 
 
diff --git a/Engine/Plugins/Experimental/HTML5Networking/Source/HTML5Networking/Classes/WebSocketNetDriver.h b/Engine/Plugins/Experimental/HTML5Networking/Source/HTML5Networking/Classes/WebSocketNetDriver.h
--- a/Engine/Plugins/Experimental/HTML5Networking/Source/HTML5Networking/Classes/WebSocketNetDriver.h
+++ b/Engine/Plugins/Experimental/HTML5Networking/Source/HTML5Networking/Classes/WebSocketNetDriver.h
@@ -9,7 +9,9 @@
 UCLASS(transient, config = Engine)
 class HTML5NETWORKING_API UWebSocketNetDriver : public UNetDriver
 {
-	GENERATED_UCLASS_BODY()
+	GENERATED_BODY()
+public:
+	UWebSocketNetDriver(const FObjectInitializer& ObjectInitializer = FObjectInitializer::Get());
 
 	/** Websocket server port*/
 	UPROPERTY(Config)
diff --git a/Engine/Plugins/Experimental/LiveEditor/Source/LiveEditor/Classes/K2Node_LiveEditObject.h b/Engine/Plugins/Experimental/LiveEditor/Source/LiveEditor/Classes/K2Node_LiveEditObject.h
--- a/Engine/Plugins/Experimental/LiveEditor/Source/LiveEditor/Classes/K2Node_LiveEditObject.h
+++ b/Engine/Plugins/Experimental/LiveEditor/Source/LiveEditor/Classes/K2Node_LiveEditObject.h
@@ -9,7 +9,9 @@
 UCLASS(MinimalAPI)
 class UK2Node_LiveEditObject : public UK2Node
 {
-	GENERATED_UCLASS_BODY()
+	GENERATED_BODY()
+public:
+	LIVEEDITOR_API UK2Node_LiveEditObject(const FObjectInitializer& ObjectInitializer = FObjectInitializer::Get());
 
 #if WITH_EDITOR
 	// Begin UEdGraphNode interface.
diff --git a/Engine/Plugins/Experimental/LiveEditor/Source/LiveEditor/Classes/LiveEditorBlueprint.h b/Engine/Plugins/Experimental/LiveEditor/Source/LiveEditor/Classes/LiveEditorBlueprint.h
--- a/Engine/Plugins/Experimental/LiveEditor/Source/LiveEditor/Classes/LiveEditorBlueprint.h
+++ b/Engine/Plugins/Experimental/LiveEditor/Source/LiveEditor/Classes/LiveEditorBlueprint.h
@@ -8,7 +8,9 @@
 UCLASS(Abstract, Blueprintable)
 class ULiveEditorBlueprint : public UObject
 {
-	GENERATED_UCLASS_BODY()
+	GENERATED_BODY()
+public:
+	ULiveEditorBlueprint(const FObjectInitializer& ObjectInitializer = FObjectInitializer::Get());
 
 public:
 	UFUNCTION(BlueprintImplementableEvent)
diff --git a/Engine/Plugins/Experimental/LiveEditor/Source/LiveEditor/Classes/LiveEditorBroadcaster.h b/Engine/Plugins/Experimental/LiveEditor/Source/LiveEditor/Classes/LiveEditorBroadcaster.h
--- a/Engine/Plugins/Experimental/LiveEditor/Source/LiveEditor/Classes/LiveEditorBroadcaster.h
+++ b/Engine/Plugins/Experimental/LiveEditor/Source/LiveEditor/Classes/LiveEditorBroadcaster.h
@@ -11,7 +11,9 @@ DECLARE_DYNAMIC_MULTICAST_DELEGATE_ThreeParams( FLiveEditorEventMIDIMultiCast, i
 UCLASS(MinimalAPI)
 class ULiveEditorBroadcaster : public UObject
 {
-	GENERATED_UCLASS_BODY()
+	GENERATED_BODY()
+public:
+	LIVEEDITOR_API ULiveEditorBroadcaster(const FObjectInitializer& ObjectInitializer = FObjectInitializer::Get());
 
 	UPROPERTY(BlueprintAssignable)
 	FLiveEditorEventMIDIMultiCast OnEventMIDI;
diff --git a/Engine/Plugins/Experimental/LiveEditor/Source/LiveEditor/Classes/LiveEditorKismetLibrary.h b/Engine/Plugins/Experimental/LiveEditor/Source/LiveEditor/Classes/LiveEditorKismetLibrary.h
--- a/Engine/Plugins/Experimental/LiveEditor/Source/LiveEditor/Classes/LiveEditorKismetLibrary.h
+++ b/Engine/Plugins/Experimental/LiveEditor/Source/LiveEditor/Classes/LiveEditorKismetLibrary.h
@@ -9,7 +9,9 @@
 UCLASS()
 class ULiveEditorKismetLibrary : public UBlueprintFunctionLibrary
 {
-	GENERATED_UCLASS_BODY()
+	GENERATED_BODY()
+public:
+	ULiveEditorKismetLibrary(const FObjectInitializer& ObjectInitializer = FObjectInitializer::Get());
 
 	UFUNCTION(BlueprintPure, Category="Development|LiveEditor")
 	static float LiveModifyFloat( float InFloat );
diff --git a/Engine/Plugins/Experimental/LiveEditor/Source/LiveEditor/Classes/LiveEditorTypes.h b/Engine/Plugins/Experimental/LiveEditor/Source/LiveEditor/Classes/LiveEditorTypes.h
--- a/Engine/Plugins/Experimental/LiveEditor/Source/LiveEditor/Classes/LiveEditorTypes.h
+++ b/Engine/Plugins/Experimental/LiveEditor/Source/LiveEditor/Classes/LiveEditorTypes.h
@@ -20,7 +20,7 @@ namespace ELiveEditControllerType UMETA(BlueprintType="true")
 USTRUCT()
 struct FLiveEditBinding
 {
-	GENERATED_USTRUCT_BODY()
+	GENERATED_BODY()
 
 	UPROPERTY()
 	int32 Channel;
@@ -51,7 +51,7 @@ public:
 USTRUCT(BlueprintType)
 struct FLiveEditorCheckpointData
 {
-	GENERATED_USTRUCT_BODY()
+	GENERATED_BODY()
 
 	UPROPERTY(EditAnywhere, BlueprintReadWrite, Category="Live Editor")
 	int32 IntValue;
@@ -65,10 +65,11 @@ struct FLiveEditorCheckpointData
 	}
 };
 
-UCLASS(Abstract, CustomConstructor)
+UCLASS(Abstract)
 class ULiveEditorTypes : public UObject
 {
-	GENERATED_UCLASS_BODY()
+	GENERATED_BODY()
+public:
 
 	ULiveEditorTypes(const FObjectInitializer& ObjectInitializer) : Super(ObjectInitializer) {}
 };
diff --git a/Engine/Plugins/Experimental/Phya/Source/Phya/Classes/PhyaCollisionHandler.h b/Engine/Plugins/Experimental/Phya/Source/Phya/Classes/PhyaCollisionHandler.h
--- a/Engine/Plugins/Experimental/Phya/Source/Phya/Classes/PhyaCollisionHandler.h
+++ b/Engine/Plugins/Experimental/Phya/Source/Phya/Classes/PhyaCollisionHandler.h
@@ -28,7 +28,9 @@ class FPhyaPairInfo : public TSharedFromThis<FPhyaPairInfo>
 UCLASS()
 class UPhyaCollisionHandler : public UPhysicsCollisionHandler
 {
-	GENERATED_UCLASS_BODY()
+	GENERATED_BODY()
+public:
+	UPhyaCollisionHandler(const FObjectInitializer& ObjectInitializer = FObjectInitializer::Get());
 
 	// Begin UPhysicsCollisionHandler interface
 	virtual void InitCollisionHandler() override;
diff --git a/Engine/Plugins/Media/MediaPlayerEditor/Source/MediaPlayerEditor/Private/Factories/MediaPlayerFactory.h b/Engine/Plugins/Media/MediaPlayerEditor/Source/MediaPlayerEditor/Private/Factories/MediaPlayerFactory.h
--- a/Engine/Plugins/Media/MediaPlayerEditor/Source/MediaPlayerEditor/Private/Factories/MediaPlayerFactory.h
+++ b/Engine/Plugins/Media/MediaPlayerEditor/Source/MediaPlayerEditor/Private/Factories/MediaPlayerFactory.h
@@ -12,7 +12,9 @@ UCLASS(hidecategories=Object)
 class UMediaPlayerFactory
 	: public UFactory
 {
-	GENERATED_UCLASS_BODY()
+	GENERATED_BODY()
+public:
+	UMediaPlayerFactory(const FObjectInitializer& ObjectInitializer = FObjectInitializer::Get());
 
 public:
 
diff --git a/Engine/Plugins/Media/MediaPlayerEditor/Source/MediaPlayerEditor/Private/Factories/MediaPlayerFactoryNew.h b/Engine/Plugins/Media/MediaPlayerEditor/Source/MediaPlayerEditor/Private/Factories/MediaPlayerFactoryNew.h
--- a/Engine/Plugins/Media/MediaPlayerEditor/Source/MediaPlayerEditor/Private/Factories/MediaPlayerFactoryNew.h
+++ b/Engine/Plugins/Media/MediaPlayerEditor/Source/MediaPlayerEditor/Private/Factories/MediaPlayerFactoryNew.h
@@ -12,7 +12,9 @@ UCLASS(hidecategories=Object)
 class UMediaPlayerFactoryNew
 	: public UFactory
 {
-	GENERATED_UCLASS_BODY()
+	GENERATED_BODY()
+public:
+	UMediaPlayerFactoryNew(const FObjectInitializer& ObjectInitializer = FObjectInitializer::Get());
 
 public:
 
diff --git a/Engine/Plugins/Media/MediaPlayerEditor/Source/MediaPlayerEditor/Private/Factories/MediaSoundWaveFactoryNew.h b/Engine/Plugins/Media/MediaPlayerEditor/Source/MediaPlayerEditor/Private/Factories/MediaSoundWaveFactoryNew.h
--- a/Engine/Plugins/Media/MediaPlayerEditor/Source/MediaPlayerEditor/Private/Factories/MediaSoundWaveFactoryNew.h
+++ b/Engine/Plugins/Media/MediaPlayerEditor/Source/MediaPlayerEditor/Private/Factories/MediaSoundWaveFactoryNew.h
@@ -12,7 +12,9 @@ UCLASS(hidecategories=Object, MinimalAPI)
 class UMediaSoundWaveFactoryNew
 	: public UFactory
 {
-	GENERATED_UCLASS_BODY()
+	GENERATED_BODY()
+public:
+	MEDIAPLAYEREDITOR_API UMediaSoundWaveFactoryNew(const FObjectInitializer& ObjectInitializer = FObjectInitializer::Get());
 
 	/** An initial UMediaPlayer asset to place in the newly created sound wave. */
 	UPROPERTY()
diff --git a/Engine/Plugins/Media/MediaPlayerEditor/Source/MediaPlayerEditor/Private/Factories/MediaTextureFactoryNew.h b/Engine/Plugins/Media/MediaPlayerEditor/Source/MediaPlayerEditor/Private/Factories/MediaTextureFactoryNew.h
--- a/Engine/Plugins/Media/MediaPlayerEditor/Source/MediaPlayerEditor/Private/Factories/MediaTextureFactoryNew.h
+++ b/Engine/Plugins/Media/MediaPlayerEditor/Source/MediaPlayerEditor/Private/Factories/MediaTextureFactoryNew.h
@@ -12,7 +12,9 @@ UCLASS(hidecategories=Object, MinimalAPI)
 class UMediaTextureFactoryNew
 	: public UFactory
 {
-	GENERATED_UCLASS_BODY()
+	GENERATED_BODY()
+public:
+	MEDIAPLAYEREDITOR_API UMediaTextureFactoryNew(const FObjectInitializer& ObjectInitializer = FObjectInitializer::Get());
 
 	/** An initial media player asset to place in the newly created texture. */
 	UPROPERTY()
diff --git a/Engine/Plugins/Messaging/ShmMessaging/Source/ShmMessaging/Private/Shared/ShmMessagingSettings.h b/Engine/Plugins/Messaging/ShmMessaging/Source/ShmMessaging/Private/Shared/ShmMessagingSettings.h
--- a/Engine/Plugins/Messaging/ShmMessaging/Source/ShmMessaging/Private/Shared/ShmMessagingSettings.h
+++ b/Engine/Plugins/Messaging/ShmMessaging/Source/ShmMessaging/Private/Shared/ShmMessagingSettings.h
@@ -9,7 +9,9 @@ UCLASS(config=Engine)
 class UShmMessagingSettings
 	: public UObject
 {
-	GENERATED_UCLASS_BODY()
+	GENERATED_BODY()
+public:
+	UShmMessagingSettings(const FObjectInitializer& ObjectInitializer = FObjectInitializer::Get());
 
 public:
 
diff --git a/Engine/Plugins/Messaging/UdpMessaging/Source/UdpMessaging/Private/Shared/UdpMessagingSettings.h b/Engine/Plugins/Messaging/UdpMessaging/Source/UdpMessaging/Private/Shared/UdpMessagingSettings.h
--- a/Engine/Plugins/Messaging/UdpMessaging/Source/UdpMessaging/Private/Shared/UdpMessagingSettings.h
+++ b/Engine/Plugins/Messaging/UdpMessaging/Source/UdpMessaging/Private/Shared/UdpMessagingSettings.h
@@ -9,7 +9,9 @@ UCLASS(config=Engine)
 class UUdpMessagingSettings
 	: public UObject
 {
-	GENERATED_UCLASS_BODY()
+	GENERATED_BODY()
+public:
+	UUdpMessagingSettings(const FObjectInitializer& ObjectInitializer = FObjectInitializer::Get());
 
 public:
 
diff --git a/Engine/Plugins/Runtime/Analytics/AnalyticsBlueprintLibrary/Source/AnalyticsBlueprintLibrary/Classes/AnalyticsBlueprintLibrary.h b/Engine/Plugins/Runtime/Analytics/AnalyticsBlueprintLibrary/Source/AnalyticsBlueprintLibrary/Classes/AnalyticsBlueprintLibrary.h
--- a/Engine/Plugins/Runtime/Analytics/AnalyticsBlueprintLibrary/Source/AnalyticsBlueprintLibrary/Classes/AnalyticsBlueprintLibrary.h
+++ b/Engine/Plugins/Runtime/Analytics/AnalyticsBlueprintLibrary/Source/AnalyticsBlueprintLibrary/Classes/AnalyticsBlueprintLibrary.h
@@ -8,7 +8,7 @@
 USTRUCT(BlueprintType)
 struct FAnalyticsEventAttr
 {
-	GENERATED_USTRUCT_BODY();
+	GENERATED_BODY()
 
 	UPROPERTY(EditAnywhere, BlueprintReadWrite, Category="Analytics")
 	FString Name;
@@ -21,7 +21,9 @@ UCLASS()
 class UAnalyticsBlueprintLibrary :
 	public UBlueprintFunctionLibrary
 {
-	GENERATED_UCLASS_BODY()
+	GENERATED_BODY()
+public:
+	UAnalyticsBlueprintLibrary(const FObjectInitializer& ObjectInitializer = FObjectInitializer::Get());
 
 	/** Starts an analytics session without any custom attributes specified */
 	UFUNCTION(BlueprintCallable, Category="Analytics")
diff --git a/Engine/Plugins/Runtime/CableComponent/Source/CableComponent/Classes/CableActor.h b/Engine/Plugins/Runtime/CableComponent/Source/CableComponent/Classes/CableActor.h
--- a/Engine/Plugins/Runtime/CableComponent/Source/CableComponent/Classes/CableActor.h
+++ b/Engine/Plugins/Runtime/CableComponent/Source/CableComponent/Classes/CableActor.h
@@ -8,7 +8,9 @@
 UCLASS(hidecategories=(Input,Collision,Replication), showcategories=("Input|MouseInput", "Input|TouchInput"))
 class ACableActor : public AActor
 {
-	GENERATED_UCLASS_BODY()
+	GENERATED_BODY()
+public:
+	ACableActor(const FObjectInitializer& ObjectInitializer = FObjectInitializer::Get());
 
 	/** Cable component that performs simulation and rendering */
 	UPROPERTY(Category=Cable, VisibleAnywhere, BlueprintReadOnly)
diff --git a/Engine/Plugins/Runtime/CableComponent/Source/CableComponent/Classes/CableComponent.h b/Engine/Plugins/Runtime/CableComponent/Source/CableComponent/Classes/CableComponent.h
--- a/Engine/Plugins/Runtime/CableComponent/Source/CableComponent/Classes/CableComponent.h
+++ b/Engine/Plugins/Runtime/CableComponent/Source/CableComponent/Classes/CableComponent.h
@@ -25,7 +25,9 @@ struct FCableParticle
 UCLASS(hidecategories=(Object, Physics, Collision, Activation, "Components|Activation"), editinlinenew, meta=(BlueprintSpawnableComponent), ClassGroup=Rendering)
 class UCableComponent : public UMeshComponent
 {
-	GENERATED_UCLASS_BODY()
+	GENERATED_BODY()
+public:
+	UCableComponent(const FObjectInitializer& ObjectInitializer = FObjectInitializer::Get());
 
 public:
 
diff --git a/Engine/Plugins/Runtime/CustomMeshComponent/Source/CustomMeshComponent/Classes/CustomMeshComponent.h b/Engine/Plugins/Runtime/CustomMeshComponent/Source/CustomMeshComponent/Classes/CustomMeshComponent.h
--- a/Engine/Plugins/Runtime/CustomMeshComponent/Source/CustomMeshComponent/Classes/CustomMeshComponent.h
+++ b/Engine/Plugins/Runtime/CustomMeshComponent/Source/CustomMeshComponent/Classes/CustomMeshComponent.h
@@ -7,7 +7,7 @@
 USTRUCT(BlueprintType)
 struct FCustomMeshTriangle
 {
-	GENERATED_USTRUCT_BODY()
+	GENERATED_BODY()
 
 	UPROPERTY(EditAnywhere, BlueprintReadWrite, Category=Triangle)
 	FVector Vertex0;
@@ -23,7 +23,9 @@ struct FCustomMeshTriangle
 UCLASS(hidecategories=(Object,LOD, Physics, Collision), editinlinenew, meta=(BlueprintSpawnableComponent), ClassGroup=Rendering)
 class CUSTOMMESHCOMPONENT_API UCustomMeshComponent : public UMeshComponent
 {
-	GENERATED_UCLASS_BODY()
+	GENERATED_BODY()
+public:
+	UCustomMeshComponent(const FObjectInitializer& ObjectInitializer = FObjectInitializer::Get());
 
 	/** Set the geometry to use on this triangle mesh */
 	UFUNCTION(BlueprintCallable, Category="Components|CustomMesh")
diff --git a/Engine/Plugins/Runtime/LeapMotionController/Source/LeapMotionController/Public/LeapMotionBoneActor.h b/Engine/Plugins/Runtime/LeapMotionController/Source/LeapMotionController/Public/LeapMotionBoneActor.h
--- a/Engine/Plugins/Runtime/LeapMotionController/Source/LeapMotionController/Public/LeapMotionBoneActor.h
+++ b/Engine/Plugins/Runtime/LeapMotionController/Source/LeapMotionController/Public/LeapMotionBoneActor.h
@@ -16,7 +16,9 @@ class ULeapMotionControllerComponent;
 UCLASS(ClassGroup = LeapMotion, BlueprintType, Blueprintable)
 class LEAPMOTIONCONTROLLER_API ALeapMotionBoneActor : public AActor
 {
-	GENERATED_UCLASS_BODY()
+	GENERATED_BODY()
+public:
+	ALeapMotionBoneActor(const FObjectInitializer& ObjectInitializer = FObjectInitializer::Get());
 
 	UPROPERTY(VisibleAnywhere, BlueprintReadOnly, Category = LeapMotion)
 	ELeapBone LeapBoneType;
diff --git a/Engine/Plugins/Runtime/LeapMotionController/Source/LeapMotionController/Public/LeapMotionControllerActor.h b/Engine/Plugins/Runtime/LeapMotionController/Source/LeapMotionController/Public/LeapMotionControllerActor.h
--- a/Engine/Plugins/Runtime/LeapMotionController/Source/LeapMotionController/Public/LeapMotionControllerActor.h
+++ b/Engine/Plugins/Runtime/LeapMotionController/Source/LeapMotionController/Public/LeapMotionControllerActor.h
@@ -15,7 +15,9 @@
 UCLASS(ClassGroup = LeapMotion, BlueprintType, Blueprintable)
 class LEAPMOTIONCONTROLLER_API ALeapMotionControllerActor : public AActor
 {
-	GENERATED_UCLASS_BODY()
+	GENERATED_BODY()
+public:
+	ALeapMotionControllerActor(const FObjectInitializer& ObjectInitializer = FObjectInitializer::Get());
 
 	/** Scale of hands, relatively to real-world size. This ignores actor's transform's scale */
 	UPROPERTY(VisibleAnywhere, BlueprintReadOnly, Category = LeapMotion)
diff --git a/Engine/Plugins/Runtime/LeapMotionController/Source/LeapMotionController/Public/LeapMotionControllerComponent.h b/Engine/Plugins/Runtime/LeapMotionController/Source/LeapMotionController/Public/LeapMotionControllerComponent.h
--- a/Engine/Plugins/Runtime/LeapMotionController/Source/LeapMotionController/Public/LeapMotionControllerComponent.h
+++ b/Engine/Plugins/Runtime/LeapMotionController/Source/LeapMotionController/Public/LeapMotionControllerComponent.h
@@ -19,7 +19,9 @@ class ALeapMotionHandActor;
 UCLASS(ClassGroup = LeapMotion, BlueprintType, Blueprintable, meta = (BlueprintSpawnableComponent))
 class LEAPMOTIONCONTROLLER_API ULeapMotionControllerComponent : public USceneComponent
 {
-	GENERATED_UCLASS_BODY()
+	GENERATED_BODY()
+public:
+	ULeapMotionControllerComponent(const FObjectInitializer& ObjectInitializer = FObjectInitializer::Get());
 
 	/** Scale of hands, relatively to real-world size. This ignores actor's transform's scale */
 	UPROPERTY(EditAnywhere, BlueprintReadWrite, Category = LeapMotion)
diff --git a/Engine/Plugins/Runtime/LeapMotionController/Source/LeapMotionController/Public/LeapMotionFunctionLibrary.h b/Engine/Plugins/Runtime/LeapMotionController/Source/LeapMotionController/Public/LeapMotionFunctionLibrary.h
--- a/Engine/Plugins/Runtime/LeapMotionController/Source/LeapMotionController/Public/LeapMotionFunctionLibrary.h
+++ b/Engine/Plugins/Runtime/LeapMotionController/Source/LeapMotionController/Public/LeapMotionFunctionLibrary.h
@@ -18,7 +18,9 @@ class ULeapMotionControllerComponent;
 UCLASS()
 class LEAPMOTIONCONTROLLER_API ULeapMotionFunctionLibrary : public UBlueprintFunctionLibrary
 {
-	GENERATED_UCLASS_BODY()
+	GENERATED_BODY()
+public:
+	ULeapMotionFunctionLibrary(const FObjectInitializer& ObjectInitializer = FObjectInitializer::Get());
 
 	/** 
 	 * Checks whether a Leap Motion controller is connected.
diff --git a/Engine/Plugins/Runtime/LeapMotionController/Source/LeapMotionController/Public/LeapMotionHandActor.h b/Engine/Plugins/Runtime/LeapMotionController/Source/LeapMotionController/Public/LeapMotionHandActor.h
--- a/Engine/Plugins/Runtime/LeapMotionController/Source/LeapMotionController/Public/LeapMotionHandActor.h
+++ b/Engine/Plugins/Runtime/LeapMotionController/Source/LeapMotionController/Public/LeapMotionHandActor.h
@@ -18,7 +18,9 @@ class ULeapMotionControllerComponent;
 UCLASS(ClassGroup = LeapMotion, BlueprintType, Blueprintable)
 class LEAPMOTIONCONTROLLER_API ALeapMotionHandActor : public AActor
 {
-	GENERATED_UCLASS_BODY()
+	GENERATED_BODY()
+public:
+	ALeapMotionHandActor(const FObjectInitializer& ObjectInitializer = FObjectInitializer::Get());
 
 	/** Scale of the hand, relatively to real-world size */
 	UPROPERTY(VisibleAnywhere, BlueprintReadOnly, Category = LeapMotion)
diff --git a/Engine/Plugins/Runtime/SoundMod/Source/SoundMod/Classes/SoundMod.h b/Engine/Plugins/Runtime/SoundMod/Source/SoundMod/Classes/SoundMod.h
--- a/Engine/Plugins/Runtime/SoundMod/Source/SoundMod/Classes/SoundMod.h
+++ b/Engine/Plugins/Runtime/SoundMod/Source/SoundMod/Classes/SoundMod.h
@@ -13,7 +13,9 @@ struct FActiveSound;
 UCLASS(hidecategories=Object, MinimalAPI, BlueprintType)
 class USoundMod : public USoundBase
 {
-	GENERATED_UCLASS_BODY()
+	GENERATED_BODY()
+public:
+	SOUNDMOD_API USoundMod(const FObjectInitializer& ObjectInitializer = FObjectInitializer::Get());
 
 	/** If set, when played directly (not through a sound cue) the nid will be played looping. */
 	UPROPERTY(EditAnywhere, Category=SoundMod)
diff --git a/Engine/Plugins/Runtime/SoundMod/Source/SoundMod/Classes/SoundModWave.h b/Engine/Plugins/Runtime/SoundMod/Source/SoundMod/Classes/SoundModWave.h
--- a/Engine/Plugins/Runtime/SoundMod/Source/SoundMod/Classes/SoundModWave.h
+++ b/Engine/Plugins/Runtime/SoundMod/Source/SoundMod/Classes/SoundModWave.h
@@ -11,7 +11,9 @@
 UCLASS()
 class USoundModWave : public USoundWaveStreaming
 {
-	GENERATED_UCLASS_BODY()
+	GENERATED_BODY()
+public:
+	USoundModWave(const FObjectInitializer& ObjectInitializer = FObjectInitializer::Get());
 
 	UPROPERTY()
 	class USoundMod* SoundMod;
@@ -22,4 +24,4 @@ class USoundModWave : public USoundWaveStreaming
 	virtual int32 GeneratePCMData(uint8* PCMData, const int32 SamplesNeeded) override;
 	// End USoundWave interface.
 
-};
\ No newline at end of file
+};
diff --git a/Engine/Plugins/Runtime/SoundMod/Source/SoundMod/Public/SoundNodeModPlayer.h b/Engine/Plugins/Runtime/SoundMod/Source/SoundMod/Public/SoundNodeModPlayer.h
--- a/Engine/Plugins/Runtime/SoundMod/Source/SoundMod/Public/SoundNodeModPlayer.h
+++ b/Engine/Plugins/Runtime/SoundMod/Source/SoundMod/Public/SoundNodeModPlayer.h
@@ -11,7 +11,9 @@
 UCLASS(hidecategories=Object, editinlinenew, MinimalAPI, meta=( DisplayName="Mod Player" ))
 class USoundNodeModPlayer : public USoundNode
 {
-	GENERATED_UCLASS_BODY()
+	GENERATED_BODY()
+public:
+	SOUNDMOD_API USoundNodeModPlayer(const FObjectInitializer& ObjectInitializer = FObjectInitializer::Get());
 
 	UPROPERTY(EditAnywhere, Category=ModPlayer)
 	USoundMod* SoundMod;
diff --git a/Engine/Plugins/Runtime/SoundMod/Source/SoundModImporter/Classes/SoundModImporterFactory.h b/Engine/Plugins/Runtime/SoundMod/Source/SoundModImporter/Classes/SoundModImporterFactory.h
--- a/Engine/Plugins/Runtime/SoundMod/Source/SoundModImporter/Classes/SoundModImporterFactory.h
+++ b/Engine/Plugins/Runtime/SoundMod/Source/SoundModImporter/Classes/SoundModImporterFactory.h
@@ -8,7 +8,9 @@
 UCLASS()
 class USoundModImporterFactory : public UFactory
 {
-	GENERATED_UCLASS_BODY()
+	GENERATED_BODY()
+public:
+	USoundModImporterFactory(const FObjectInitializer& ObjectInitializer = FObjectInitializer::Get());
 
 	// UFactory interface
 	virtual UObject* FactoryCreateBinary(UClass* Class, UObject* InParent, FName Name, EObjectFlags Flags, UObject* Context, const TCHAR* Type, const uint8*& Buffer, const uint8* BufferEnd, FFeedbackContext* Warn) override;
diff --git a/Engine/Plugins/Runtime/TwitchLiveStreaming/Source/TwitchLiveStreaming/TwitchProjectSettings.h b/Engine/Plugins/Runtime/TwitchLiveStreaming/Source/TwitchLiveStreaming/TwitchProjectSettings.h
--- a/Engine/Plugins/Runtime/TwitchLiveStreaming/Source/TwitchLiveStreaming/TwitchProjectSettings.h
+++ b/Engine/Plugins/Runtime/TwitchLiveStreaming/Source/TwitchLiveStreaming/TwitchProjectSettings.h
@@ -6,7 +6,9 @@
 UCLASS(Config=Engine, DefaultConfig)
 class UTwitchProjectSettings : public UObject
 {
-	GENERATED_UCLASS_BODY()
+	GENERATED_BODY()
+public:
+	UTwitchProjectSettings(const FObjectInitializer& ObjectInitializer = FObjectInitializer::Get());
 
 public:
 
diff --git a/Engine/Plugins/Runtime/WebBrowserAsPlugin/Source/WebBrowserAsPlugin/Public/WebBrowser.h b/Engine/Plugins/Runtime/WebBrowserAsPlugin/Source/WebBrowserAsPlugin/Public/WebBrowser.h
--- a/Engine/Plugins/Runtime/WebBrowserAsPlugin/Source/WebBrowserAsPlugin/Public/WebBrowser.h
+++ b/Engine/Plugins/Runtime/WebBrowserAsPlugin/Source/WebBrowserAsPlugin/Public/WebBrowser.h
@@ -14,7 +14,9 @@
 UCLASS()
 class WEBBROWSERASPLUGIN_API UWebBrowser : public UWidget
 {
-	GENERATED_UCLASS_BODY()
+	GENERATED_BODY()
+public:
+	UWebBrowser(const FObjectInitializer& ObjectInitializer = FObjectInitializer::Get());
 
 public:
 
diff --git a/Engine/Plugins/ScriptPlugin/Source/ScriptEditorPlugin/Classes/ReimportScriptFactory.h b/Engine/Plugins/ScriptPlugin/Source/ScriptEditorPlugin/Classes/ReimportScriptFactory.h
--- a/Engine/Plugins/ScriptPlugin/Source/ScriptEditorPlugin/Classes/ReimportScriptFactory.h
+++ b/Engine/Plugins/ScriptPlugin/Source/ScriptEditorPlugin/Classes/ReimportScriptFactory.h
@@ -14,7 +14,9 @@
 UCLASS(MinimalApi, collapsecategories, EarlyAccessPreview)
 class UReimportScriptFactory : public UScriptFactory, public FReimportHandler
 {
-	GENERATED_UCLASS_BODY()
+	GENERATED_BODY()
+public:
+	SCRIPTEDITORPLUGIN_API UReimportScriptFactory(const FObjectInitializer& ObjectInitializer = FObjectInitializer::Get());
 
 	UPROPERTY()
 	class UScriptBlueprint* OriginalScript;
diff --git a/Engine/Plugins/ScriptPlugin/Source/ScriptEditorPlugin/Classes/ScriptFactory.h b/Engine/Plugins/ScriptPlugin/Source/ScriptEditorPlugin/Classes/ScriptFactory.h
--- a/Engine/Plugins/ScriptPlugin/Source/ScriptEditorPlugin/Classes/ScriptFactory.h
+++ b/Engine/Plugins/ScriptPlugin/Source/ScriptEditorPlugin/Classes/ScriptFactory.h
@@ -11,7 +11,9 @@
 UCLASS(collapsecategories, hidecategories = Object, EarlyAccessPreview)
 class SCRIPTEDITORPLUGIN_API UScriptFactory : public UFactory
 {
-	GENERATED_UCLASS_BODY()
+	GENERATED_BODY()
+public:
+	UScriptFactory(const FObjectInitializer& ObjectInitializer = FObjectInitializer::Get());
 
 public:
 
diff --git a/Engine/Plugins/ScriptPlugin/Source/ScriptPlugin/Classes/ScriptBlueprint.h b/Engine/Plugins/ScriptPlugin/Source/ScriptPlugin/Classes/ScriptBlueprint.h
--- a/Engine/Plugins/ScriptPlugin/Source/ScriptPlugin/Classes/ScriptBlueprint.h
+++ b/Engine/Plugins/ScriptPlugin/Source/ScriptPlugin/Classes/ScriptBlueprint.h
@@ -12,7 +12,9 @@
 UCLASS(BlueprintType, EarlyAccessPreview)
 class SCRIPTPLUGIN_API UScriptBlueprint : public UBlueprint
 {
-	GENERATED_UCLASS_BODY()
+	GENERATED_BODY()
+public:
+	UScriptBlueprint(const FObjectInitializer& ObjectInitializer = FObjectInitializer::Get());
 
 public:
 
@@ -54,4 +56,4 @@ public:
 	 */
 	void UpdateSourceCodeIfChanged();
 #endif
-};
\ No newline at end of file
+};
diff --git a/Engine/Plugins/ScriptPlugin/Source/ScriptPlugin/Classes/ScriptBlueprintGeneratedClass.h b/Engine/Plugins/ScriptPlugin/Source/ScriptPlugin/Classes/ScriptBlueprintGeneratedClass.h
--- a/Engine/Plugins/ScriptPlugin/Source/ScriptPlugin/Classes/ScriptBlueprintGeneratedClass.h
+++ b/Engine/Plugins/ScriptPlugin/Source/ScriptPlugin/Classes/ScriptBlueprintGeneratedClass.h
@@ -117,7 +117,9 @@ public:
 UCLASS(EarlyAccessPreview)
 class SCRIPTPLUGIN_API UScriptBlueprintGeneratedClass : public UBlueprintGeneratedClass
 {
-	GENERATED_UCLASS_BODY()
+	GENERATED_BODY()
+public:
+	UScriptBlueprintGeneratedClass(const FObjectInitializer& ObjectInitializer = FObjectInitializer::Get());
 
 	/** Generated script bytecode */
 	UPROPERTY()
@@ -159,4 +161,4 @@ class SCRIPTPLUGIN_API UScriptBlueprintGeneratedClass : public UBlueprintGenerat
 		}
 		return ScriptClass;
 	}
-};
\ No newline at end of file
+};
diff --git a/Engine/Plugins/ScriptPlugin/Source/ScriptPlugin/Classes/ScriptContext.h b/Engine/Plugins/ScriptPlugin/Source/ScriptPlugin/Classes/ScriptContext.h
--- a/Engine/Plugins/ScriptPlugin/Source/ScriptPlugin/Classes/ScriptContext.h
+++ b/Engine/Plugins/ScriptPlugin/Source/ScriptPlugin/Classes/ScriptContext.h
@@ -10,7 +10,9 @@
 UCLASS(DefaultToInstanced, EarlyAccessPreview)
 class SCRIPTPLUGIN_API UScriptContext : public UObject
 {
-	GENERATED_UCLASS_BODY()
+	GENERATED_BODY()
+public:
+	UScriptContext(const FObjectInitializer& ObjectInitializer = FObjectInitializer::Get());
 
 	virtual void PostInitProperties() override;
 	virtual void BeginDestroy() override;
@@ -22,4 +24,4 @@ protected:
 
 	/** Script context (code) */
 	FScriptContextBase* Context;
-};
\ No newline at end of file
+};
diff --git a/Engine/Plugins/ScriptPlugin/Source/ScriptPlugin/Classes/ScriptContextComponent.h b/Engine/Plugins/ScriptPlugin/Source/ScriptPlugin/Classes/ScriptContextComponent.h
--- a/Engine/Plugins/ScriptPlugin/Source/ScriptPlugin/Classes/ScriptContextComponent.h
+++ b/Engine/Plugins/ScriptPlugin/Source/ScriptPlugin/Classes/ScriptContextComponent.h
@@ -10,7 +10,9 @@
 UCLASS(EarlyAccessPreview)
 class SCRIPTPLUGIN_API UScriptContextComponent : public UActorComponent
 {
-	GENERATED_UCLASS_BODY()
+	GENERATED_BODY()
+public:
+	UScriptContextComponent(const FObjectInitializer& ObjectInitializer = FObjectInitializer::Get());
 
 public:
 
diff --git a/Engine/Plugins/ScriptPlugin/Source/ScriptPlugin/Classes/ScriptPluginComponent.h b/Engine/Plugins/ScriptPlugin/Source/ScriptPlugin/Classes/ScriptPluginComponent.h
--- a/Engine/Plugins/ScriptPlugin/Source/ScriptPlugin/Classes/ScriptPluginComponent.h
+++ b/Engine/Plugins/ScriptPlugin/Source/ScriptPlugin/Classes/ScriptPluginComponent.h
@@ -10,7 +10,9 @@
 UCLASS(Blueprintable, hidecategories = (Object, ActorComponent), editinlinenew, meta = (BlueprintSpawnableComponent), ClassGroup = Script, Abstract, EarlyAccessPreview)
 class SCRIPTPLUGIN_API UScriptPluginComponent : public UActorComponent
 {
-	GENERATED_UCLASS_BODY()
+	GENERATED_BODY()
+public:
+	UScriptPluginComponent(const FObjectInitializer& ObjectInitializer = FObjectInitializer::Get());
 
 public:
 
diff --git a/Engine/Plugins/ScriptPlugin/Source/ScriptPlugin/Classes/ScriptTestActor.h b/Engine/Plugins/ScriptPlugin/Source/ScriptPlugin/Classes/ScriptTestActor.h
--- a/Engine/Plugins/ScriptPlugin/Source/ScriptPlugin/Classes/ScriptTestActor.h
+++ b/Engine/Plugins/ScriptPlugin/Source/ScriptPlugin/Classes/ScriptTestActor.h
@@ -8,7 +8,9 @@
 UCLASS(BlueprintType, MinimalAPI, EarlyAccessPreview, notplaceable)
 class AScriptTestActor : public AActor
 {
-	GENERATED_UCLASS_BODY()
+	GENERATED_BODY()
+public:
+	SCRIPTPLUGIN_API AScriptTestActor(const FObjectInitializer& ObjectInitializer = FObjectInitializer::Get());
 
 	UPROPERTY(Category = Test, EditAnywhere, BlueprintReadWrite)
 	FString TestString;
diff --git a/Engine/Plugins/Slate/SlateRemote/Source/SlateRemote/Private/Shared/SlateRemoteSettings.h b/Engine/Plugins/Slate/SlateRemote/Source/SlateRemote/Private/Shared/SlateRemoteSettings.h
--- a/Engine/Plugins/Slate/SlateRemote/Source/SlateRemote/Private/Shared/SlateRemoteSettings.h
+++ b/Engine/Plugins/Slate/SlateRemote/Source/SlateRemote/Private/Shared/SlateRemoteSettings.h
@@ -12,7 +12,9 @@ UCLASS(config=Engine)
 class USlateRemoteSettings
 	: public UObject
 {
-	GENERATED_UCLASS_BODY()
+	GENERATED_BODY()
+public:
+	USlateRemoteSettings(const FObjectInitializer& ObjectInitializer = FObjectInitializer::Get());
 
 public:
 
diff --git a/Engine/Source/Developer/AITestSuite/Classes/Actions/TestPawnAction_CallFunction.h b/Engine/Source/Developer/AITestSuite/Classes/Actions/TestPawnAction_CallFunction.h
--- a/Engine/Source/Developer/AITestSuite/Classes/Actions/TestPawnAction_CallFunction.h
+++ b/Engine/Source/Developer/AITestSuite/Classes/Actions/TestPawnAction_CallFunction.h
@@ -9,7 +9,9 @@
 UCLASS()
 class UTestPawnAction_CallFunction : public UTestPawnAction_Log
 {
-	GENERATED_UCLASS_BODY()
+	GENERATED_BODY()
+public:
+	UTestPawnAction_CallFunction(const FObjectInitializer& ObjectInitializer = FObjectInitializer::Get());
 
 	typedef void(*FFunctionToCall)(UPawnActionsComponent& ActionsComponent, UTestPawnAction_CallFunction& Caller, ETestPawnActionMessage::Type);
 	FFunctionToCall FunctionToCall;
@@ -21,4 +23,4 @@ class UTestPawnAction_CallFunction : public UTestPawnAction_Log
 	virtual bool Resume() override;
 	virtual void OnFinished(EPawnActionResult::Type WithResult) override;
 	virtual void OnChildFinished(UPawnAction& Action, EPawnActionResult::Type WithResult) override;
-};
\ No newline at end of file
+};
diff --git a/Engine/Source/Developer/AITestSuite/Classes/Actions/TestPawnAction_Log.h b/Engine/Source/Developer/AITestSuite/Classes/Actions/TestPawnAction_Log.h
--- a/Engine/Source/Developer/AITestSuite/Classes/Actions/TestPawnAction_Log.h
+++ b/Engine/Source/Developer/AITestSuite/Classes/Actions/TestPawnAction_Log.h
@@ -20,7 +20,9 @@ namespace ETestPawnActionMessage
 UCLASS()
 class UTestPawnAction_Log : public UPawnAction
 {
-	GENERATED_UCLASS_BODY()
+	GENERATED_BODY()
+public:
+	UTestPawnAction_Log(const FObjectInitializer& ObjectInitializer = FObjectInitializer::Get());
 
 	FTestLogger<int32>* Logger;
 
diff --git a/Engine/Source/Developer/AITestSuite/Classes/BehaviorTree/TestBTDecorator_CantExecute.h b/Engine/Source/Developer/AITestSuite/Classes/BehaviorTree/TestBTDecorator_CantExecute.h
--- a/Engine/Source/Developer/AITestSuite/Classes/BehaviorTree/TestBTDecorator_CantExecute.h
+++ b/Engine/Source/Developer/AITestSuite/Classes/BehaviorTree/TestBTDecorator_CantExecute.h
@@ -7,7 +7,9 @@
 UCLASS(meta=(HiddenNode))
 class UTestBTDecorator_CantExecute : public UBTDecorator
 {
-	GENERATED_UCLASS_BODY()
+	GENERATED_BODY()
+public:
+	UTestBTDecorator_CantExecute(const FObjectInitializer& ObjectInitializer = FObjectInitializer::Get());
 
 	virtual bool CalculateRawConditionValue(UBehaviorTreeComponent& OwnerComp, uint8* NodeMemory) const override;
 };
diff --git a/Engine/Source/Developer/AITestSuite/Classes/BehaviorTree/TestBTDecorator_DelayedAbort.h b/Engine/Source/Developer/AITestSuite/Classes/BehaviorTree/TestBTDecorator_DelayedAbort.h
--- a/Engine/Source/Developer/AITestSuite/Classes/BehaviorTree/TestBTDecorator_DelayedAbort.h
+++ b/Engine/Source/Developer/AITestSuite/Classes/BehaviorTree/TestBTDecorator_DelayedAbort.h
@@ -12,7 +12,9 @@ struct FBTDelayedAbortMemory
 UCLASS(meta = (HiddenNode))
 class UTestBTDecorator_DelayedAbort : public UBTDecorator
 {
-	GENERATED_UCLASS_BODY()
+	GENERATED_BODY()
+public:
+	UTestBTDecorator_DelayedAbort(const FObjectInitializer& ObjectInitializer = FObjectInitializer::Get());
 
 	UPROPERTY()
 	int32 DelayTicks;
diff --git a/Engine/Source/Developer/AITestSuite/Classes/BehaviorTree/TestBTTask_LatentWithFlags.h b/Engine/Source/Developer/AITestSuite/Classes/BehaviorTree/TestBTTask_LatentWithFlags.h
--- a/Engine/Source/Developer/AITestSuite/Classes/BehaviorTree/TestBTTask_LatentWithFlags.h
+++ b/Engine/Source/Developer/AITestSuite/Classes/BehaviorTree/TestBTTask_LatentWithFlags.h
@@ -15,7 +15,9 @@ struct FBTLatentTaskMemory
 UCLASS(meta = (HiddenNode))
 class UTestBTTask_LatentWithFlags : public UBTTaskNode
 {
-	GENERATED_UCLASS_BODY()
+	GENERATED_BODY()
+public:
+	UTestBTTask_LatentWithFlags(const FObjectInitializer& ObjectInitializer = FObjectInitializer::Get());
 
 	UPROPERTY()
 	int32 LogIndexExecuteStart;
diff --git a/Engine/Source/Developer/AITestSuite/Classes/BehaviorTree/TestBTTask_Log.h b/Engine/Source/Developer/AITestSuite/Classes/BehaviorTree/TestBTTask_Log.h
--- a/Engine/Source/Developer/AITestSuite/Classes/BehaviorTree/TestBTTask_Log.h
+++ b/Engine/Source/Developer/AITestSuite/Classes/BehaviorTree/TestBTTask_Log.h
@@ -12,7 +12,9 @@ struct FBTLogTaskMemory
 UCLASS(meta=(HiddenNode))
 class UTestBTTask_Log : public UBTTaskNode
 {
-	GENERATED_UCLASS_BODY()
+	GENERATED_BODY()
+public:
+	UTestBTTask_Log(const FObjectInitializer& ObjectInitializer = FObjectInitializer::Get());
 
 	UPROPERTY()
 	int32 LogIndex;
diff --git a/Engine/Source/Developer/AITestSuite/Classes/BehaviorTree/TestBTTask_SetFlag.h b/Engine/Source/Developer/AITestSuite/Classes/BehaviorTree/TestBTTask_SetFlag.h
--- a/Engine/Source/Developer/AITestSuite/Classes/BehaviorTree/TestBTTask_SetFlag.h
+++ b/Engine/Source/Developer/AITestSuite/Classes/BehaviorTree/TestBTTask_SetFlag.h
@@ -7,7 +7,9 @@
 UCLASS(meta=(HiddenNode))
 class UTestBTTask_SetFlag : public UBTTaskNode
 {
-	GENERATED_UCLASS_BODY()
+	GENERATED_BODY()
+public:
+	UTestBTTask_SetFlag(const FObjectInitializer& ObjectInitializer = FObjectInitializer::Get());
 
 	UPROPERTY()
 	FName KeyName;
diff --git a/Engine/Source/Developer/AITestSuite/Classes/MockAI.h b/Engine/Source/Developer/AITestSuite/Classes/MockAI.h
--- a/Engine/Source/Developer/AITestSuite/Classes/MockAI.h
+++ b/Engine/Source/Developer/AITestSuite/Classes/MockAI.h
@@ -20,7 +20,9 @@ struct FTestTickHelper : FTickableGameObject
 UCLASS()
 class UMockAI : public UObject
 {
-	GENERATED_UCLASS_BODY()
+	GENERATED_BODY()
+public:
+	UMockAI(const FObjectInitializer& ObjectInitializer = FObjectInitializer::Get());
 
 	virtual ~UMockAI();
 
diff --git a/Engine/Source/Developer/AITestSuite/Classes/MockAI_BT.h b/Engine/Source/Developer/AITestSuite/Classes/MockAI_BT.h
--- a/Engine/Source/Developer/AITestSuite/Classes/MockAI_BT.h
+++ b/Engine/Source/Developer/AITestSuite/Classes/MockAI_BT.h
@@ -7,7 +7,9 @@
 UCLASS()
 class UMockAI_BT : public UMockAI
 {
-	GENERATED_UCLASS_BODY()
+	GENERATED_BODY()
+public:
+	UMockAI_BT(const FObjectInitializer& ObjectInitializer = FObjectInitializer::Get());
 
 	UPROPERTY()
 	UBehaviorTreeComponent* BTComp;
diff --git a/Engine/Source/Developer/Android/AndroidPlatformEditor/Private/AndroidSDKSettings.h b/Engine/Source/Developer/Android/AndroidPlatformEditor/Private/AndroidSDKSettings.h
--- a/Engine/Source/Developer/Android/AndroidPlatformEditor/Private/AndroidSDKSettings.h
+++ b/Engine/Source/Developer/Android/AndroidPlatformEditor/Private/AndroidSDKSettings.h
@@ -13,8 +13,10 @@
 UCLASS(config=Engine, globaluserconfig)
 class ANDROIDPLATFORMEDITOR_API UAndroidSDKSettings : public UObject
 {
+	GENERATED_BODY()
+public:
+	UAndroidSDKSettings(const FObjectInitializer& ObjectInitializer = FObjectInitializer::Get());
 public:
-	GENERATED_UCLASS_BODY()
 
 	// Location on disk of the Android SDK (falls back to ANDROID_HOME environment variable if this is left blank)
 	UPROPERTY(GlobalConfig, EditAnywhere, Category = SDKConfig, Meta = (DisplayName = "Location of Android SDK (the directory usually contains 'android-sdk-')"))
diff --git a/Engine/Source/Developer/FriendsAndChat/Public/FriendsAndChatStyle.h b/Engine/Source/Developer/FriendsAndChat/Public/FriendsAndChatStyle.h
--- a/Engine/Source/Developer/FriendsAndChat/Public/FriendsAndChatStyle.h
+++ b/Engine/Source/Developer/FriendsAndChat/Public/FriendsAndChatStyle.h
@@ -21,7 +21,7 @@ USTRUCT()
 struct FRIENDSANDCHAT_API FFriendsAndChatStyle
 	: public FSlateWidgetStyle
 {
-	GENERATED_USTRUCT_BODY()
+	GENERATED_BODY()
 
 	// Default Constructor
 	FFriendsAndChatStyle() { }
diff --git a/Engine/Source/Developer/FunctionalTesting/Classes/FuncTestRenderingComponent.h b/Engine/Source/Developer/FunctionalTesting/Classes/FuncTestRenderingComponent.h
--- a/Engine/Source/Developer/FunctionalTesting/Classes/FuncTestRenderingComponent.h
+++ b/Engine/Source/Developer/FunctionalTesting/Classes/FuncTestRenderingComponent.h
@@ -7,7 +7,9 @@
 UCLASS(hidecategories=Object, editinlinenew)
 class UFuncTestRenderingComponent : public UPrimitiveComponent
 {
-	GENERATED_UCLASS_BODY()
+	GENERATED_BODY()
+public:
+	UFuncTestRenderingComponent(const FObjectInitializer& ObjectInitializer = FObjectInitializer::Get());
 
 	// Begin UPrimitiveComponent Interface
 	virtual FPrimitiveSceneProxy* CreateSceneProxy() override;
@@ -17,4 +19,4 @@ class UFuncTestRenderingComponent : public UPrimitiveComponent
 	// Begin USceneComponent Interface
 	virtual FBoxSphereBounds CalcBounds(const FTransform &LocalToWorld) const override;
 	// End USceneComponent Interface
-};
\ No newline at end of file
+};
diff --git a/Engine/Source/Developer/FunctionalTesting/Classes/FunctionalAITest.h b/Engine/Source/Developer/FunctionalTesting/Classes/FunctionalAITest.h
--- a/Engine/Source/Developer/FunctionalTesting/Classes/FunctionalAITest.h
+++ b/Engine/Source/Developer/FunctionalTesting/Classes/FunctionalAITest.h
@@ -14,7 +14,7 @@ DECLARE_DYNAMIC_MULTICAST_DELEGATE_TwoParams(FFunctionalTestAISpawned, AAIContro
 USTRUCT()
 struct FAITestSpawnInfo
 {
-	GENERATED_USTRUCT_BODY()
+	GENERATED_BODY()
 
 	/** Determines AI to be spawned */
 	UPROPERTY(EditAnywhere, BlueprintReadWrite, Category=AISpawn)
@@ -56,7 +56,7 @@ struct FAITestSpawnInfo
 USTRUCT()
 struct FPendingDelayedSpawn : public FAITestSpawnInfo
 {
-	GENERATED_USTRUCT_BODY()
+	GENERATED_BODY()
 
 	uint32 NumberToSpawnLeft;
 	float TimeToNextSpawn;
@@ -73,7 +73,7 @@ struct FPendingDelayedSpawn : public FAITestSpawnInfo
 USTRUCT()
 struct FAITestSpawnSet
 {
-	GENERATED_USTRUCT_BODY()
+	GENERATED_BODY()
 
 	/** what to spawn */
 	UPROPERTY(EditAnywhere, BlueprintReadWrite, Category=AISpawn)
@@ -97,7 +97,9 @@ struct FAITestSpawnSet
 UCLASS(Blueprintable, MinimalAPI)
 class AFunctionalAITest : public AFunctionalTest
 {
-	GENERATED_UCLASS_BODY()
+	GENERATED_BODY()
+public:
+	FUNCTIONALTESTING_API AFunctionalAITest(const FObjectInitializer& ObjectInitializer = FObjectInitializer::Get());
 	
 protected:
 	UPROPERTY(EditAnywhere, BlueprintReadOnly, Category=AITest)
@@ -149,4 +151,4 @@ protected:
 
 	void KillOffSpawnedPawns();
 	void ClearPendingDelayedSpawns();
-};
\ No newline at end of file
+};
diff --git a/Engine/Source/Developer/FunctionalTesting/Classes/FunctionalTest.h b/Engine/Source/Developer/FunctionalTesting/Classes/FunctionalTest.h
--- a/Engine/Source/Developer/FunctionalTesting/Classes/FunctionalTest.h
+++ b/Engine/Source/Developer/FunctionalTesting/Classes/FunctionalTest.h
@@ -26,7 +26,9 @@ DECLARE_DELEGATE_OneParam(FFunctionalTestDoneSignature, class AFunctionalTest*);
 UCLASS(Blueprintable, MinimalAPI)
 class AFunctionalTest : public AActor
 {
-	GENERATED_UCLASS_BODY()
+	GENERATED_BODY()
+public:
+	FUNCTIONALTESTING_API AFunctionalTest(const FObjectInitializer& ObjectInitializer = FObjectInitializer::Get());
 
 	static const uint32 DefaultTimeLimit = 60;	// seconds
 
diff --git a/Engine/Source/Developer/FunctionalTesting/Classes/FunctionalTestingManager.h b/Engine/Source/Developer/FunctionalTesting/Classes/FunctionalTestingManager.h
--- a/Engine/Source/Developer/FunctionalTesting/Classes/FunctionalTestingManager.h
+++ b/Engine/Source/Developer/FunctionalTesting/Classes/FunctionalTestingManager.h
@@ -17,7 +17,9 @@ namespace FFunctionalTesting
 UCLASS(BlueprintType, MinimalAPI)
 class UFunctionalTestingManager : public UBlueprintFunctionLibrary
 {
-	GENERATED_UCLASS_BODY()
+	GENERATED_BODY()
+public:
+	FUNCTIONALTESTING_API UFunctionalTestingManager(const FObjectInitializer& ObjectInitializer = FObjectInitializer::Get());
 	
 	UPROPERTY(Transient)
 	TArray<class AFunctionalTest*> TestsLeft;
diff --git a/Engine/Source/Developer/GameplayDebugger/Classes/GameplayDebuggerSettings.h b/Engine/Source/Developer/GameplayDebugger/Classes/GameplayDebuggerSettings.h
--- a/Engine/Source/Developer/GameplayDebugger/Classes/GameplayDebuggerSettings.h
+++ b/Engine/Source/Developer/GameplayDebugger/Classes/GameplayDebuggerSettings.h
@@ -13,7 +13,7 @@ GameplayDebuggerSettings.h: Declares the UGameplayDebuggerSettings class.
 USTRUCT()
 struct FGDTCustomViewNames
 {
-	GENERATED_USTRUCT_BODY()
+	GENERATED_BODY()
 
 	UPROPERTY(EditAnywhere, config, Category = "Settings")
 	FString GameView1;
@@ -34,7 +34,9 @@ struct FGDTCustomViewNames
 UCLASS(config = EditorUserSettings)
 class GAMEPLAYDEBUGGER_API UGameplayDebuggerSettings : public UObject
 {
-	GENERATED_UCLASS_BODY()
+	GENERATED_BODY()
+public:
+	UGameplayDebuggerSettings(const FObjectInitializer& ObjectInitializer = FObjectInitializer::Get());
 public:
 #if WITH_EDITOR
 	// Begin UObject Interface
diff --git a/Engine/Source/Developer/GameplayDebugger/Classes/GameplayDebuggingComponent.h b/Engine/Source/Developer/GameplayDebugger/Classes/GameplayDebuggingComponent.h
--- a/Engine/Source/Developer/GameplayDebugger/Classes/GameplayDebuggingComponent.h
+++ b/Engine/Source/Developer/GameplayDebugger/Classes/GameplayDebuggingComponent.h
@@ -36,7 +36,9 @@ DECLARE_MULTICAST_DELEGATE_TwoParams(FOnDebuggingTargetChanged, class AActor* /*
 UCLASS(config=Engine)
 class GAMEPLAYDEBUGGER_API UGameplayDebuggingComponent : public UPrimitiveComponent, public IEQSQueryResultSourceInterface
 {
-	GENERATED_UCLASS_BODY()
+	GENERATED_BODY()
+public:
+	UGameplayDebuggingComponent(const FObjectInitializer& ObjectInitializer = FObjectInitializer::Get());
 
 	friend class AGameplayDebuggingHUDComponent;
 
@@ -167,11 +169,15 @@ class GAMEPLAYDEBUGGER_API UGameplayDebuggingComponent : public UPrimitiveCompon
 	UFUNCTION(exec)
 	void ServerReplicateData(uint32 InMessage, uint32 DataView);
 
-	UFUNCTION(reliable, server, WithValidation)
+	UFUNCTION(reliable, server="ServerCollectNavmeshData_Implementation", WithValidation="ServerCollectNavmeshData_Validate")
 	void ServerCollectNavmeshData(FVector_NetQuantize10 TargetLocation);
+	virtual void ServerCollectNavmeshData_Implementation(FVector_NetQuantize10 TargetLocation);
+	virtual bool ServerCollectNavmeshData_Validate(FVector_NetQuantize10 );
 
-	UFUNCTION(reliable, server, WithValidation)
+	UFUNCTION(reliable, server="ServerDiscardNavmeshData_Implementation", WithValidation="ServerDiscardNavmeshData_Validate")
 	void ServerDiscardNavmeshData();
+	virtual void ServerDiscardNavmeshData_Implementation();
+	virtual bool ServerDiscardNavmeshData_Validate();
 
 	void PrepareNavMeshData(struct FNavMeshSceneProxyData*) const;
 
@@ -194,8 +200,10 @@ class GAMEPLAYDEBUGGER_API UGameplayDebuggingComponent : public UPrimitiveCompon
 	UPROPERTY(Replicated)
 	AActor* TargetActor;
 	
-	UFUNCTION(Reliable, Client, WithValidation)
+	UFUNCTION(Reliable, Client="ClientEnableTargetSelection_Implementation", WithValidation="ClientEnableTargetSelection_Validate")
 	void ClientEnableTargetSelection(bool bEnable);
+	virtual void ClientEnableTargetSelection_Implementation(bool bEnable);
+	virtual bool ClientEnableTargetSelection_Validate(bool );
 
 	void SetActorToDebug(AActor* Actor);
 	FORCEINLINE AActor* GetSelectedActor() const
diff --git a/Engine/Source/Developer/GameplayDebugger/Classes/GameplayDebuggingControllerComponent.h b/Engine/Source/Developer/GameplayDebugger/Classes/GameplayDebuggingControllerComponent.h
--- a/Engine/Source/Developer/GameplayDebugger/Classes/GameplayDebuggingControllerComponent.h
+++ b/Engine/Source/Developer/GameplayDebugger/Classes/GameplayDebuggingControllerComponent.h
@@ -15,7 +15,9 @@ class AGameplayDebuggingReplicator;
 UCLASS(config = Engine)
 class GAMEPLAYDEBUGGER_API UGameplayDebuggingControllerComponent : public UActorComponent
 {
-	GENERATED_UCLASS_BODY()
+	GENERATED_BODY()
+public:
+	UGameplayDebuggingControllerComponent(const FObjectInitializer& ObjectInitializer = FObjectInitializer::Get());
 
 #if WITH_HOT_RELOAD_CTORS
 	/** DO NOT USE. This constructor is for internal usage only for hot-reload purposes. */
diff --git a/Engine/Source/Developer/GameplayDebugger/Classes/GameplayDebuggingHUDComponent.h b/Engine/Source/Developer/GameplayDebugger/Classes/GameplayDebuggingHUDComponent.h
--- a/Engine/Source/Developer/GameplayDebugger/Classes/GameplayDebuggingHUDComponent.h
+++ b/Engine/Source/Developer/GameplayDebugger/Classes/GameplayDebuggingHUDComponent.h
@@ -24,7 +24,9 @@ struct GAMEPLAYDEBUGGER_API FDebugCategoryView
 UCLASS(config = Engine, notplaceable)
 class GAMEPLAYDEBUGGER_API AGameplayDebuggingHUDComponent : public AActor
 {
-	GENERATED_UCLASS_BODY()
+	GENERATED_BODY()
+public:
+	AGameplayDebuggingHUDComponent(const FObjectInitializer& ObjectInitializer = FObjectInitializer::Get());
 
 	UPROPERTY(config)
 	float MenuStartX;
diff --git a/Engine/Source/Developer/GameplayDebugger/Classes/GameplayDebuggingReplicator.h b/Engine/Source/Developer/GameplayDebugger/Classes/GameplayDebuggingReplicator.h
--- a/Engine/Source/Developer/GameplayDebugger/Classes/GameplayDebuggingReplicator.h
+++ b/Engine/Source/Developer/GameplayDebugger/Classes/GameplayDebuggingReplicator.h
@@ -15,7 +15,9 @@ DECLARE_MULTICAST_DELEGATE_OneParam(FOnSelectionChanged, class AActor*);
 UCLASS(config = Engine, NotBlueprintable, Transient, hidecategories = Actor, notplaceable)
 class GAMEPLAYDEBUGGER_API AGameplayDebuggingReplicator : public AActor
 {
-	GENERATED_UCLASS_BODY()
+	GENERATED_BODY()
+public:
+	AGameplayDebuggingReplicator(const FObjectInitializer& ObjectInitializer = FObjectInitializer::Get());
 
 	UPROPERTY(config)
 	FString DebugComponentClassName;
@@ -86,17 +88,24 @@ class GAMEPLAYDEBUGGER_API AGameplayDebuggingReplicator : public AActor
 	UPROPERTY()
 	class UTexture2D* DefaultTexture_Green;
 
-	UFUNCTION(reliable, server, WithValidation)
+	UFUNCTION(reliable, server="ServerReplicateMessage_Implementation", WithValidation="ServerReplicateMessage_Validate")
 	void ServerReplicateMessage(class AActor* Actor, uint32 InMessage, uint32 DataView = 0);
+	virtual void ServerReplicateMessage_Implementation(AActor* Actor, uint32 InMessage, uint32 DataView);
+	virtual bool ServerReplicateMessage_Validate(AActor* , uint32 , uint32 );
 
-	UFUNCTION(reliable, client, WithValidation)
+	UFUNCTION(reliable, client="ClientReplicateMessage_Implementation", WithValidation="ClientReplicateMessage_Validate")
 	void ClientReplicateMessage(class AActor* Actor, uint32 InMessage, uint32 DataView = 0);
+	virtual void ClientReplicateMessage_Implementation(AActor* Actor, uint32 InMessage, uint32 DataView);
+	virtual bool ClientReplicateMessage_Validate(AActor* , uint32 , uint32 );
 
-	UFUNCTION(Reliable, Client)
+	UFUNCTION(Reliable, Client="ClientAutoActivate_Implementation")
 	void ClientAutoActivate();
+	virtual void ClientAutoActivate_Implementation();
 
-	UFUNCTION(Reliable, Client, WithValidation)
+	UFUNCTION(Reliable, Client="ClientEnableTargetSelection_Implementation", WithValidation="ClientEnableTargetSelection_Validate")
 	void ClientEnableTargetSelection(bool bEnable, APlayerController* Context);
+	virtual void ClientEnableTargetSelection_Implementation(bool bEnable, APlayerController* Context);
+	virtual bool ClientEnableTargetSelection_Validate(bool , APlayerController* );
 
 #if WITH_EDITOR
 	void PostEditChangeProperty(FPropertyChangedEvent& PropertyChangedEvent) override;
@@ -132,8 +141,10 @@ class GAMEPLAYDEBUGGER_API AGameplayDebuggingReplicator : public AActor
 
 	FORCEINLINE AActor* GetSelectedActorToDebug() { return LastSelectedActorToDebug; }
 	
-	UFUNCTION(reliable, server, WithValidation)
+	UFUNCTION(reliable, server="ServerSetActorToDebug_Implementation", WithValidation="ServerSetActorToDebug_Validate")
 	void ServerSetActorToDebug(AActor* InActor);
+	virtual void ServerSetActorToDebug_Implementation(AActor* InActor);
+	virtual bool ServerSetActorToDebug_Validate(AActor* );
 
 	/**
 	 * Iterates through the pawn list to find the next pawn of the specified type to debug
diff --git a/Engine/Source/Developer/HTML5/HTML5PlatformEditor/Classes/HTML5TargetSettings.h b/Engine/Source/Developer/HTML5/HTML5PlatformEditor/Classes/HTML5TargetSettings.h
--- a/Engine/Source/Developer/HTML5/HTML5PlatformEditor/Classes/HTML5TargetSettings.h
+++ b/Engine/Source/Developer/HTML5/HTML5PlatformEditor/Classes/HTML5TargetSettings.h
@@ -16,9 +16,11 @@ UCLASS(config=Engine, defaultconfig)
 class HTML5PLATFORMEDITOR_API UHTML5TargetSettings
 	: public UObject
 {
+	GENERATED_BODY()
+public:
+	UHTML5TargetSettings(const FObjectInitializer& ObjectInitializer = FObjectInitializer::Get());
 public:
 
-	GENERATED_UCLASS_BODY()
 
  	/**
  	 * Setting to control HTML5 Heap size (in Development)
diff --git a/Engine/Source/Developer/HTML5/HTML5PlatformEditor/Private/HTML5SDKSettings.h b/Engine/Source/Developer/HTML5/HTML5PlatformEditor/Private/HTML5SDKSettings.h
--- a/Engine/Source/Developer/HTML5/HTML5PlatformEditor/Private/HTML5SDKSettings.h
+++ b/Engine/Source/Developer/HTML5/HTML5PlatformEditor/Private/HTML5SDKSettings.h
@@ -14,7 +14,7 @@ class IHTML5TargetPlatformModule;
 USTRUCT()
 struct FHTML5DeviceMapping
 {
-	GENERATED_USTRUCT_BODY()
+	GENERATED_BODY()
 
 	UPROPERTY(EditAnywhere, Category = HTML5_Devices, Meta = (DisplayName = "Browser name"))
 	FString DeviceName;
@@ -26,7 +26,7 @@ struct FHTML5DeviceMapping
 USTRUCT()
 struct FHTML5SDKPath
 {
-	GENERATED_USTRUCT_BODY()
+	GENERATED_BODY()
 
 	UPROPERTY(EditAnywhere, Category = SDK_Path)
 	FString SDKPath;
@@ -42,8 +42,10 @@ struct FHTML5SDKPath
 UCLASS(config=Engine, globaluserconfig)
 class HTML5PLATFORMEDITOR_API UHTML5SDKSettings : public UObject
 {
+	GENERATED_BODY()
+public:
+	UHTML5SDKSettings(const FObjectInitializer& ObjectInitializer = FObjectInitializer::Get());
 public:
-	GENERATED_UCLASS_BODY()
 
 	UPROPERTY(GlobalConfig, EditAnywhere, Category = HTML5_SDK_Paths, Meta = (DisplayName = "Location of Python exe (This can be blank if python is on the PATH)"))
 	FFilePath Python;
diff --git a/Engine/Source/Developer/InternationalizationSettings/Classes/InternationalizationSettingsModel.h b/Engine/Source/Developer/InternationalizationSettings/Classes/InternationalizationSettingsModel.h
--- a/Engine/Source/Developer/InternationalizationSettings/Classes/InternationalizationSettingsModel.h
+++ b/Engine/Source/Developer/InternationalizationSettings/Classes/InternationalizationSettingsModel.h
@@ -12,7 +12,9 @@ UCLASS()
 class INTERNATIONALIZATIONSETTINGS_API UInternationalizationSettingsModel
 	:	public UObject
 {
-	GENERATED_UCLASS_BODY()
+	GENERATED_BODY()
+public:
+	UInternationalizationSettingsModel(const FObjectInitializer& ObjectInitializer = FObjectInitializer::Get());
 
 public:
 
diff --git a/Engine/Source/Developer/Linux/LinuxTargetPlatform/Classes/LinuxTargetSettings.h b/Engine/Source/Developer/Linux/LinuxTargetPlatform/Classes/LinuxTargetSettings.h
--- a/Engine/Source/Developer/Linux/LinuxTargetPlatform/Classes/LinuxTargetSettings.h
+++ b/Engine/Source/Developer/Linux/LinuxTargetPlatform/Classes/LinuxTargetSettings.h
@@ -16,9 +16,10 @@ UCLASS(config=Engine, defaultconfig)
 class LINUXTARGETPLATFORM_API ULinuxTargetSettings
 	: public UObject
 {
-public:
+	GENERATED_BODY()
 
-	GENERATED_UCLASS_BODY()
+public:
+	ULinuxTargetSettings(const FObjectInitializer& ObjectInitializer = FObjectInitializer::Get());
 
 	/**
 	 * The collection of RHI's we want to support on this platform.
diff --git a/Engine/Source/Developer/LogVisualizer/Private/VisualLoggerCameraController.h b/Engine/Source/Developer/LogVisualizer/Private/VisualLoggerCameraController.h
--- a/Engine/Source/Developer/LogVisualizer/Private/VisualLoggerCameraController.h
+++ b/Engine/Source/Developer/LogVisualizer/Private/VisualLoggerCameraController.h
@@ -8,11 +8,13 @@
 UCLASS(config=Input, hidedropdown)
 class AVisualLoggerCameraController : public ADebugCameraController
 {
+	GENERATED_BODY()
+public:
+	AVisualLoggerCameraController(const FObjectInitializer& ObjectInitializer = FObjectInitializer::Get());
 public:
 	DECLARE_DELEGATE_OneParam(FActorSelectedDelegate, class AActor*);
 	DECLARE_DELEGATE_OneParam(FLogEntryIterationDelegate, int32);
 
-	GENERATED_UCLASS_BODY()
 
 	UPROPERTY()
 	class AActor* PickedActor;
diff --git a/Engine/Source/Developer/LogVisualizer/Private/VisualLoggerHUD.h b/Engine/Source/Developer/LogVisualizer/Private/VisualLoggerHUD.h
--- a/Engine/Source/Developer/LogVisualizer/Private/VisualLoggerHUD.h
+++ b/Engine/Source/Developer/LogVisualizer/Private/VisualLoggerHUD.h
@@ -7,7 +7,9 @@
 UCLASS()
 class AVisualLoggerHUD : public ADebugCameraHUD
 {
-	GENERATED_UCLASS_BODY()
+	GENERATED_BODY()
+public:
+	AVisualLoggerHUD(const FObjectInitializer& ObjectInitializer = FObjectInitializer::Get());
 
 	FFontRenderInfo TextRenderInfo;
 
diff --git a/Engine/Source/Developer/LogVisualizer/Private/VisualLoggerRenderingActor.h b/Engine/Source/Developer/LogVisualizer/Private/VisualLoggerRenderingActor.h
--- a/Engine/Source/Developer/LogVisualizer/Private/VisualLoggerRenderingActor.h
+++ b/Engine/Source/Developer/LogVisualizer/Private/VisualLoggerRenderingActor.h
@@ -13,8 +13,10 @@ DECLARE_MULTICAST_DELEGATE_OneParam(FOnSelectionChanged, class AActor*);
 UCLASS(config = Engine, NotBlueprintable, Transient)
 class LOGVISUALIZER_API AVisualLoggerRenderingActor : public AActor
 {
+	GENERATED_BODY()
+public:
+	AVisualLoggerRenderingActor(const FObjectInitializer& ObjectInitializer = FObjectInitializer::Get());
 public:
-	GENERATED_UCLASS_BODY()
 
 	void OnItemSelectionChanged(const FVisualLogDevice::FVisualLogEntryItem& EntryItem);
 	void ObjectSelectionChanged(TSharedPtr<class STimeline> TimeLine);
diff --git a/Engine/Source/Developer/LogVisualizer/Private/VisualLoggerRenderingComponent.h b/Engine/Source/Developer/LogVisualizer/Private/VisualLoggerRenderingComponent.h
--- a/Engine/Source/Developer/LogVisualizer/Private/VisualLoggerRenderingComponent.h
+++ b/Engine/Source/Developer/LogVisualizer/Private/VisualLoggerRenderingComponent.h
@@ -10,8 +10,10 @@
 UCLASS()
 class UVisualLoggerRenderingComponent : public UPrimitiveComponent
 {
+	GENERATED_BODY()
+public:
+	UVisualLoggerRenderingComponent(const FObjectInitializer& ObjectInitializer = FObjectInitializer::Get());
 public:
-	GENERATED_UCLASS_BODY()
 
 	virtual FPrimitiveSceneProxy* CreateSceneProxy() override;
 	virtual FBoxSphereBounds CalcBounds(const FTransform &LocalToWorld) const override;
diff --git a/Engine/Source/Developer/LogVisualizer/Public/LogVisualizerSessionSettings.h b/Engine/Source/Developer/LogVisualizer/Public/LogVisualizerSessionSettings.h
--- a/Engine/Source/Developer/LogVisualizer/Public/LogVisualizerSessionSettings.h
+++ b/Engine/Source/Developer/LogVisualizer/Public/LogVisualizerSessionSettings.h
@@ -11,7 +11,9 @@ GameplayDebuggerSettings.h: Declares the UGameplayDebuggerSettings class.
 UCLASS()
 class LOGVISUALIZER_API ULogVisualizerSessionSettings : public UObject
 {
-	GENERATED_UCLASS_BODY()
+	GENERATED_BODY()
+public:
+	ULogVisualizerSessionSettings(const FObjectInitializer& ObjectInitializer = FObjectInitializer::Get());
 public:
 	/**Whether to show trivial logs, i.e. the ones with only one entry.*/
 	UPROPERTY(EditAnywhere, Category = "VisualLogger")
diff --git a/Engine/Source/Developer/LogVisualizer/Public/LogVisualizerSettings.h b/Engine/Source/Developer/LogVisualizer/Public/LogVisualizerSettings.h
--- a/Engine/Source/Developer/LogVisualizer/Public/LogVisualizerSettings.h
+++ b/Engine/Source/Developer/LogVisualizer/Public/LogVisualizerSettings.h
@@ -11,7 +11,7 @@ GameplayDebuggerSettings.h: Declares the UGameplayDebuggerSettings class.
 USTRUCT()
 struct FCategoryFilter
 {
-	GENERATED_USTRUCT_BODY()
+	GENERATED_BODY()
 
 	UPROPERTY(config)
 	FString CategoryName;
@@ -26,7 +26,7 @@ struct FCategoryFilter
 USTRUCT()
 struct FVisualLoggerFilters
 {
-	GENERATED_USTRUCT_BODY()
+	GENERATED_BODY()
 
 	UPROPERTY(config)
 	FString SearchBoxFilter;
@@ -46,7 +46,9 @@ struct FCategoryFiltersManager;
 UCLASS(config = EditorUserSettings)
 class LOGVISUALIZER_API ULogVisualizerSettings : public UObject
 {
-	GENERATED_UCLASS_BODY()
+	GENERATED_BODY()
+public:
+	ULogVisualizerSettings(const FObjectInitializer& ObjectInitializer = FObjectInitializer::Get());
 	friend struct FCategoryFiltersManager;
 
 public:
diff --git a/Engine/Source/Developer/Mac/MacTargetPlatform/Classes/MacTargetSettings.h b/Engine/Source/Developer/Mac/MacTargetPlatform/Classes/MacTargetSettings.h
--- a/Engine/Source/Developer/Mac/MacTargetPlatform/Classes/MacTargetSettings.h
+++ b/Engine/Source/Developer/Mac/MacTargetPlatform/Classes/MacTargetSettings.h
@@ -16,8 +16,7 @@ UCLASS()
 class MACTARGETPLATFORM_API UMacTargetSettings
 	: public UObject
 {
+	GENERATED_BODY()
 public:
-
-	GENERATED_UCLASS_BODY()
-
+	UMacTargetSettings(const FObjectInitializer& ObjectInitializer = FObjectInitializer::Get());
 };
diff --git a/Engine/Source/Developer/MovieSceneTools/Classes/K2Node_PlayMovieScene.h b/Engine/Source/Developer/MovieSceneTools/Classes/K2Node_PlayMovieScene.h
--- a/Engine/Source/Developer/MovieSceneTools/Classes/K2Node_PlayMovieScene.h
+++ b/Engine/Source/Developer/MovieSceneTools/Classes/K2Node_PlayMovieScene.h
@@ -14,7 +14,7 @@ DECLARE_EVENT( UK2Node_PlayMovieScene, FOnMovieSceneBindingsChanged );
 USTRUCT()
 struct FEdGraphSchemaAction_K2AddPlayMovieScene : public FEdGraphSchemaAction_K2NewNode
 {
-	GENERATED_USTRUCT_BODY()
+	GENERATED_BODY()
 
 	FEdGraphSchemaAction_K2AddPlayMovieScene()
 		: FEdGraphSchemaAction_K2NewNode()
@@ -36,7 +36,9 @@ struct FEdGraphSchemaAction_K2AddPlayMovieScene : public FEdGraphSchemaAction_K2
 UCLASS( MinimalAPI )
 class UK2Node_PlayMovieScene : public UK2Node
 {
-	GENERATED_UCLASS_BODY()
+	GENERATED_BODY()
+public:
+	MOVIESCENETOOLS_API UK2Node_PlayMovieScene(const FObjectInitializer& ObjectInitializer = FObjectInitializer::Get());
 
 #if WITH_EDITOR
 	// UEdGraphNode interface.
diff --git a/Engine/Source/Developer/MovieSceneTools/Classes/MovieSceneFactory.h b/Engine/Source/Developer/MovieSceneTools/Classes/MovieSceneFactory.h
--- a/Engine/Source/Developer/MovieSceneTools/Classes/MovieSceneFactory.h
+++ b/Engine/Source/Developer/MovieSceneTools/Classes/MovieSceneFactory.h
@@ -7,7 +7,9 @@
 UCLASS()
 class UMovieSceneFactory : public UFactory
 {
-	GENERATED_UCLASS_BODY()
+	GENERATED_BODY()
+public:
+	UMovieSceneFactory(const FObjectInitializer& ObjectInitializer = FObjectInitializer::Get());
 
 	// Begin UFactory Interface
 	virtual UObject* FactoryCreateNew(UClass* Class,UObject* InParent,FName Name,EObjectFlags Flags,UObject* Context,FFeedbackContext* Warn) override;
diff --git a/Engine/Source/Developer/ProfilerMessages/Classes/ProfilerServiceMessages.h b/Engine/Source/Developer/ProfilerMessages/Classes/ProfilerServiceMessages.h
--- a/Engine/Source/Developer/ProfilerMessages/Classes/ProfilerServiceMessages.h
+++ b/Engine/Source/Developer/ProfilerMessages/Classes/ProfilerServiceMessages.h
@@ -10,7 +10,7 @@
 USTRUCT()
 struct FProfilerServiceAuthorize
 {
-	GENERATED_USTRUCT_BODY()
+	GENERATED_BODY()
 
 	/**
 	 */
@@ -53,7 +53,7 @@ struct TStructOpsTypeTraits<FProfilerServiceAuthorize> : public TStructOpsTypeTr
 USTRUCT()
 struct FProfilerServiceAuthorize2
 {
-	GENERATED_USTRUCT_BODY()
+	GENERATED_BODY()
 
 	/**
 	 */
@@ -102,7 +102,7 @@ struct TStructOpsTypeTraits<FProfilerServiceAuthorize2> : public TStructOpsTypeT
 USTRUCT()
 struct FProfilerServiceData
 {
-	GENERATED_USTRUCT_BODY()
+	GENERATED_BODY()
 
 	/**
 	 */
@@ -145,7 +145,7 @@ struct TStructOpsTypeTraits<FProfilerServiceData> : public TStructOpsTypeTraitsB
 USTRUCT()
 struct FProfilerServiceData2
 {
-	GENERATED_USTRUCT_BODY()
+	GENERATED_BODY()
 
 	/**
 	 */
@@ -192,7 +192,7 @@ struct TStructOpsTypeTraits<FProfilerServiceData2> : public TStructOpsTypeTraits
 USTRUCT()
 struct FProfilerServicePreviewAck
 {
-	GENERATED_USTRUCT_BODY()
+	GENERATED_BODY()
 
 	/**
 	 */
@@ -233,7 +233,7 @@ struct TStructOpsTypeTraits<FProfilerServicePreviewAck> : public TStructOpsTypeT
 USTRUCT()
 struct FProfilerServiceMetaData
 {
-	GENERATED_USTRUCT_BODY()
+	GENERATED_BODY()
 
 	/**
 	 */
@@ -277,7 +277,7 @@ struct TStructOpsTypeTraits<FProfilerServiceMetaData> : public TStructOpsTypeTra
 USTRUCT()
 struct FProfilerServiceFileChunk
 {
-	GENERATED_USTRUCT_BODY()
+	GENERATED_BODY()
 
 	/** The ID of the instance where this message should be sent. */
 	UPROPERTY()
@@ -345,7 +345,7 @@ struct TStructOpsTypeTraits<FProfilerServiceFileChunk> : public TStructOpsTypeTr
 USTRUCT()
 struct FProfilerServicePing
 {
-	GENERATED_USTRUCT_BODY()
+	GENERATED_BODY()
 };
 
 template<>
@@ -363,7 +363,7 @@ struct TStructOpsTypeTraits<FProfilerServicePing> : public TStructOpsTypeTraitsB
 USTRUCT()
 struct FProfilerServicePong
 {
-	GENERATED_USTRUCT_BODY()
+	GENERATED_BODY()
 
 };
 
@@ -382,7 +382,7 @@ struct TStructOpsTypeTraits<FProfilerServicePong> : public TStructOpsTypeTraitsB
 USTRUCT()
 struct FProfilerServiceSubscribe
 {
-	GENERATED_USTRUCT_BODY()
+	GENERATED_BODY()
 
 	/**
 	 */
@@ -423,7 +423,7 @@ struct TStructOpsTypeTraits<FProfilerServiceSubscribe> : public TStructOpsTypeTr
 USTRUCT()
 struct FProfilerServiceUnsubscribe
 {
-	GENERATED_USTRUCT_BODY()
+	GENERATED_BODY()
 
 	/**
 	 */
@@ -465,7 +465,7 @@ struct TStructOpsTypeTraits<FProfilerServiceUnsubscribe> : public TStructOpsType
 USTRUCT()
 struct FProfilerServiceCapture
 {
-	GENERATED_USTRUCT_BODY()
+	GENERATED_BODY()
 
 	/**
 	 * The data capture state that should be set.
@@ -501,7 +501,7 @@ struct TStructOpsTypeTraits<FProfilerServiceCapture> : public TStructOpsTypeTrai
 USTRUCT()
 struct FProfilerServicePreview
 {
-	GENERATED_USTRUCT_BODY()
+	GENERATED_BODY()
 
 	/**
 	 * The data preview state that should be set.
@@ -537,7 +537,7 @@ struct TStructOpsTypeTraits<FProfilerServicePreview> : public TStructOpsTypeTrai
 USTRUCT()
 struct FProfilerServiceRequest
 {
-	GENERATED_USTRUCT_BODY()
+	GENERATED_BODY()
 
 	/** Request @see EProfilerRequestType. */
 	UPROPERTY()
